Niagara take recorder integration

diff --git a/Engine/Plugins/FX/NiagaraSimCaching/Source/NiagaraSimCaching/NiagaraSimCaching.Build.cs b/Engine/Plugins/FX/NiagaraSimCaching/Source/NiagaraSimCaching/NiagaraSimCaching.Build.cs
new file mode 100644
--- /dev/null
+++ b/Engine/Plugins/FX/NiagaraSimCaching/Source/NiagaraSimCaching/NiagaraSimCaching.Build.cs
@@ -0,0 +1,39 @@
+// Copyright Epic Games, Inc. All Rights Reserved.
+
+namespace UnrealBuildTool.Rules
+{
+	public class NiagaraSimCaching : ModuleRules
+	{
+        public NiagaraSimCaching(ReadOnlyTargetRules Target) : base(Target)
+		{
+			PublicIncludePaths.Add(ModuleDirectory + "/Public");
+
+			PublicDependencyModuleNames.AddRange(
+				new[]
+				{
+					"Core",
+					"CoreUObject",
+                    "Engine",
+                    "EngineSettings",
+                    "RenderCore",
+                    "RHI",
+					"Niagara",
+					"MovieScene",
+					"MovieSceneTracks",
+					"LevelSequence",
+				});
+
+			if(Target.bBuildEditor)
+            {
+				// Slate/Editor extensions
+				PublicDependencyModuleNames.AddRange(
+					new[]
+					{
+						"UnrealEd",
+						"SlateCore",
+						"Slate"
+					});
+            }
+		}
+	}
+}
diff --git a/Engine/Plugins/FX/NiagaraSimCaching/Source/NiagaraSimCaching/Private/Niagara/NiagaraSimCachingPlugin.cpp b/Engine/Plugins/FX/NiagaraSimCaching/Source/NiagaraSimCaching/Private/Niagara/NiagaraSimCachingPlugin.cpp
new file mode 100644
--- /dev/null
+++ b/Engine/Plugins/FX/NiagaraSimCaching/Source/NiagaraSimCaching/Private/Niagara/NiagaraSimCachingPlugin.cpp
@@ -0,0 +1,17 @@
+// Copyright Epic Games, Inc. All Rights Reserved.
+
+#include "Niagara/NiagaraSimCachingPlugin.h"
+#include "CoreMinimal.h"
+#include "Modules/ModuleManager.h"
+
+IMPLEMENT_MODULE(INiagaraSimCachingPlugin, NiagaraSimCaching)
+
+DEFINE_LOG_CATEGORY(LogNiagaraSimCaching)
+
+void INiagaraSimCachingPlugin::StartupModule()
+{
+}
+
+void INiagaraSimCachingPlugin::ShutdownModule() 
+{
+}
diff --git a/Engine/Plugins/FX/NiagaraSimCaching/Source/NiagaraSimCaching/Private/Niagara/Sequencer/MovieSceneNiagaraCacheSection.cpp b/Engine/Plugins/FX/NiagaraSimCaching/Source/NiagaraSimCaching/Private/Niagara/Sequencer/MovieSceneNiagaraCacheSection.cpp
new file mode 100644
--- /dev/null
+++ b/Engine/Plugins/FX/NiagaraSimCaching/Source/NiagaraSimCaching/Private/Niagara/Sequencer/MovieSceneNiagaraCacheSection.cpp
@@ -0,0 +1,31 @@
+﻿// Copyright Epic Games, Inc. All Rights Reserved.
+
+#include "Niagara/Sequencer/MovieSceneNiagaraCacheSection.h"
+#include "MovieScene.h"
+
+FMovieSceneNiagaraCacheParams::FMovieSceneNiagaraCacheParams() : FMovieSceneBaseCacheParams()
+{
+	SimCache = nullptr;
+}
+
+float FMovieSceneNiagaraCacheParams::GetSequenceLength() const
+{
+	if (SimCache)
+	{
+		return SimCache->GetDurationSeconds();
+	}
+	return 0.f;
+}
+
+UMovieSceneNiagaraCacheSection::UMovieSceneNiagaraCacheSection(const FObjectInitializer& ObjectInitializer)
+	: Super(ObjectInitializer)
+{
+	ParamsPtr = &Params;
+#if WITH_EDITOR
+	InitializePlayRate();
+#endif
+}
+
+
+
+
diff --git a/Engine/Plugins/FX/NiagaraSimCaching/Source/NiagaraSimCaching/Private/Niagara/Sequencer/MovieSceneNiagaraCacheTemplate.cpp b/Engine/Plugins/FX/NiagaraSimCaching/Source/NiagaraSimCaching/Private/Niagara/Sequencer/MovieSceneNiagaraCacheTemplate.cpp
new file mode 100644
--- /dev/null
+++ b/Engine/Plugins/FX/NiagaraSimCaching/Source/NiagaraSimCaching/Private/Niagara/Sequencer/MovieSceneNiagaraCacheTemplate.cpp
@@ -0,0 +1,89 @@
+﻿// Copyright Epic Games, Inc. All Rights Reserved.
+
+#include "MovieSceneNiagaraCacheTemplate.h"
+
+#include "NiagaraComponent.h"
+
+DECLARE_CYCLE_STAT(TEXT("Niagara Cache - Token Execute"), MovieSceneEval_NiagaraCache_TokenExecute, STATGROUP_MovieSceneEval);
+
+/** Used to set Manual Tick back to previous when outside section */
+struct FPreAnimatedNiagaraCacheTokenProducer : IMovieScenePreAnimatedTokenProducer
+{
+	virtual IMovieScenePreAnimatedTokenPtr CacheExistingState(UObject& Object) const override
+	{
+		struct FToken : IMovieScenePreAnimatedToken
+		{
+			explicit FToken(const UNiagaraComponent* InNiagaraComponent)
+			{
+				AgeMode = InNiagaraComponent->GetAgeUpdateMode();
+			}
+
+			virtual void RestoreState(UObject& ObjectToRestore, const UE::MovieScene::FRestoreStateParams& Params) override
+			{
+				UNiagaraComponent* NiagaraComponent = CastChecked<UNiagaraComponent>(&ObjectToRestore);
+				NiagaraComponent->SetAgeUpdateMode(AgeMode);
+				NiagaraComponent->SetSimCache(nullptr);
+			}
+
+			ENiagaraAgeUpdateMode AgeMode;
+		};
+
+		return FToken(CastChecked<UNiagaraComponent>(&Object));
+	}
+	static FMovieSceneAnimTypeID GetAnimTypeID()
+	{
+		return TMovieSceneAnimTypeID<FPreAnimatedNiagaraCacheTokenProducer>();
+	}
+};
+
+/** A movie scene execution token that executes a NiagaraCache */
+struct FNiagaraCacheExecutionToken : IMovieSceneExecutionToken
+{
+	FNiagaraCacheExecutionToken(const FMovieSceneNiagaraCacheSectionTemplateParameters &InParams) :
+		Params(InParams)
+	{}
+
+	/** Execute this token, operating on all objects referenced by 'Operand' */
+	virtual void Execute(const FMovieSceneContext& Context, const FMovieSceneEvaluationOperand& Operand, FPersistentEvaluationData& PersistentData, IMovieScenePlayer& Player) override
+	{
+		MOVIESCENE_DETAILED_SCOPE_CYCLE_COUNTER(MovieSceneEval_NiagaraCache_TokenExecute)
+		if (Operand.ObjectBindingID.IsValid() && Params.NiagaraCacheParams.SimCache && Params.NiagaraCacheParams.SimCache->GetNumFrames() > 0)
+		{
+			for (TWeakObjectPtr<> WeakObj : Player.FindBoundObjects(Operand))
+			{
+				if (UNiagaraComponent* NiagaraComponent = Cast<UNiagaraComponent>(WeakObj))
+				{
+					if (NiagaraComponent->GetSimCache() != Params.NiagaraCacheParams.SimCache)
+					{
+						NiagaraComponent->SetSimCache(Params.NiagaraCacheParams.SimCache);
+					}
+					
+					// calculate the time at which to evaluate the sim cache
+					float Age = Params.MapTimeToAnimation(Params.NiagaraCacheParams, Params.NiagaraCacheParams.SimCache->GetDurationSeconds(), Context.GetTime(), Context.GetFrameRate()) + Params.NiagaraCacheParams.SimCache->GetStartSeconds();
+
+					Player.SavePreAnimatedState(*NiagaraComponent, FPreAnimatedNiagaraCacheTokenProducer::GetAnimTypeID(), FPreAnimatedNiagaraCacheTokenProducer());
+					NiagaraComponent->SetDesiredAge(Age);
+					NiagaraComponent->SetAgeUpdateMode(ENiagaraAgeUpdateMode::DesiredAge);
+					if (!NiagaraComponent->IsActive())
+					{
+						NiagaraComponent->ResetSystem();
+					}
+				}
+			}
+		}
+	}
+
+	FMovieSceneNiagaraCacheSectionTemplateParameters Params;
+};
+
+FMovieSceneNiagaraCacheSectionTemplate::FMovieSceneNiagaraCacheSectionTemplate(const UMovieSceneNiagaraCacheSection& InSection)
+	: Params(const_cast<FMovieSceneNiagaraCacheParams&> (InSection.Params), InSection.GetInclusiveStartFrame(), InSection.GetExclusiveEndFrame())
+{
+}
+
+//We use a token here so we can set the manual tick state back to what it was previously when outside this section.
+//This is similar to how Skeletal Animation evaluation also works.
+void FMovieSceneNiagaraCacheSectionTemplate::Evaluate(const FMovieSceneEvaluationOperand& Operand, const FMovieSceneContext& Context, const FPersistentEvaluationData& PersistentData, FMovieSceneExecutionTokens& ExecutionTokens) const
+{
+	ExecutionTokens.Add(FNiagaraCacheExecutionToken(Params));
+}
\ No newline at end of file
diff --git a/Engine/Plugins/FX/NiagaraSimCaching/Source/NiagaraSimCaching/Private/Niagara/Sequencer/MovieSceneNiagaraCacheTemplate.h b/Engine/Plugins/FX/NiagaraSimCaching/Source/NiagaraSimCaching/Private/Niagara/Sequencer/MovieSceneNiagaraCacheTemplate.h
new file mode 100644
--- /dev/null
+++ b/Engine/Plugins/FX/NiagaraSimCaching/Source/NiagaraSimCaching/Private/Niagara/Sequencer/MovieSceneNiagaraCacheTemplate.h
@@ -0,0 +1,38 @@
+﻿// Copyright Epic Games, Inc. All Rights Reserved.
+
+#pragma once
+
+#include "Evaluation/MovieSceneBaseCacheTemplate.h"
+#include "Niagara/Sequencer/MovieSceneNiagaraCacheSection.h"
+#include "MovieSceneNiagaraCacheTemplate.generated.h"
+
+USTRUCT()
+struct FMovieSceneNiagaraCacheSectionTemplateParameters : public FMovieSceneBaseCacheSectionTemplateParameters
+{
+	GENERATED_BODY()
+
+	FMovieSceneNiagaraCacheSectionTemplateParameters() {}
+
+	FMovieSceneNiagaraCacheSectionTemplateParameters(const FMovieSceneNiagaraCacheParams& InNiagaraCacheParams, FFrameNumber InSectionStartTime, FFrameNumber InSectionEndTime)
+		: FMovieSceneBaseCacheSectionTemplateParameters(InSectionStartTime, InSectionEndTime), NiagaraCacheParams(InNiagaraCacheParams)
+ 
+	{}
+
+	UPROPERTY()
+	FMovieSceneNiagaraCacheParams NiagaraCacheParams;
+};
+
+USTRUCT()
+struct FMovieSceneNiagaraCacheSectionTemplate : public FMovieSceneEvalTemplate
+{
+	GENERATED_BODY()
+
+	FMovieSceneNiagaraCacheSectionTemplate() {}
+	FMovieSceneNiagaraCacheSectionTemplate(const UMovieSceneNiagaraCacheSection& Section);
+
+	virtual UScriptStruct& GetScriptStructImpl() const override { return *StaticStruct(); }
+	virtual void Evaluate(const FMovieSceneEvaluationOperand& Operand, const FMovieSceneContext& Context, const FPersistentEvaluationData& PersistentData, FMovieSceneExecutionTokens& ExecutionTokens) const override;
+
+	UPROPERTY()
+	FMovieSceneNiagaraCacheSectionTemplateParameters Params;
+};
diff --git a/Engine/Plugins/FX/NiagaraSimCaching/Source/NiagaraSimCaching/Private/Niagara/Sequencer/MovieSceneNiagaraCacheTrack.cpp b/Engine/Plugins/FX/NiagaraSimCaching/Source/NiagaraSimCaching/Private/Niagara/Sequencer/MovieSceneNiagaraCacheTrack.cpp
new file mode 100644
--- /dev/null
+++ b/Engine/Plugins/FX/NiagaraSimCaching/Source/NiagaraSimCaching/Private/Niagara/Sequencer/MovieSceneNiagaraCacheTrack.cpp
@@ -0,0 +1,123 @@
+﻿// Copyright Epic Games, Inc. All Rights Reserved.
+
+#include "Niagara/Sequencer/MovieSceneNiagaraCacheTrack.h"
+#include "MovieScene.h"
+#include "Niagara/Sequencer/MovieSceneNiagaraCacheSection.h"
+#include "MovieSceneNiagaraCacheTemplate.h"
+#include "NiagaraSimCache.h"
+
+#define LOCTEXT_NAMESPACE "MovieSceneNiagaraCacheTrack"
+
+/* UMovieSceneNiagaraCacheTrack constructor
+ *****************************************************************************/
+
+UMovieSceneNiagaraCacheTrack::UMovieSceneNiagaraCacheTrack(const FObjectInitializer& ObjectInitializer)
+	: Super(ObjectInitializer)
+{
+#if WITH_EDITORONLY_DATA
+	TrackTint = FColor(19, 58, 171, 80);
+#endif
+
+	SupportedBlendTypes.Add(EMovieSceneBlendType::Absolute);
+
+	EvalOptions.bCanEvaluateNearestSection = true;
+	EvalOptions.bEvaluateInPreroll = true;
+}
+
+/* UMovieSceneNiagaraCacheTrack interface
+ *****************************************************************************/
+
+UMovieSceneSection* UMovieSceneNiagaraCacheTrack::AddNewAnimation(FFrameNumber KeyTime, UNiagaraComponent* NiagaraComponent)
+{
+	UMovieSceneNiagaraCacheSection* NewSection = Cast<UMovieSceneNiagaraCacheSection>(CreateNewSection());
+	{
+		//const FFrameTime AnimationLength = NiagaraCache->CacheCollection->GetMaxDuration() * GetTypedOuter<UMovieScene>()->GetTickResolution();
+		//const int32 IFrameNumber = AnimationLength.FrameNumber.Value + (int)(AnimationLength.GetSubFrame() + 0.5f) + 1;
+		//NewSection->InitialPlacementOnRow(AnimationSections, KeyTime, IFrameNumber, INDEX_NONE);
+	}
+
+	AddSection(*NewSection);
+
+	return NewSection;
+}
+
+TArray<UMovieSceneSection*> UMovieSceneNiagaraCacheTrack::GetAnimSectionsAtTime(FFrameNumber Time)
+{
+	TArray<UMovieSceneSection*> Sections;
+	for (auto Section : AnimationSections)
+	{
+		if (Section->IsTimeWithinSection(Time))
+		{
+			Sections.Add(Section);
+		}
+	}
+
+	return Sections;
+}
+
+FMovieSceneEvalTemplatePtr UMovieSceneNiagaraCacheTrack::CreateTemplateForSection(const UMovieSceneSection& InSection) const
+{
+	return FMovieSceneNiagaraCacheSectionTemplate(*CastChecked<UMovieSceneNiagaraCacheSection>(&InSection));
+}
+
+/* UMovieSceneTrack interface
+ *****************************************************************************/
+
+const TArray<UMovieSceneSection*>& UMovieSceneNiagaraCacheTrack::GetAllSections() const
+{
+	return AnimationSections;
+}
+
+bool UMovieSceneNiagaraCacheTrack::SupportsType(TSubclassOf<UMovieSceneSection> SectionClass) const
+{
+	return SectionClass == UMovieSceneNiagaraCacheSection::StaticClass();
+}
+
+UMovieSceneSection* UMovieSceneNiagaraCacheTrack::CreateNewSection()
+{
+	UMovieSceneNiagaraCacheSection* MovieSceneNiagaraCacheSection = NewObject<UMovieSceneNiagaraCacheSection>(this, NAME_None, RF_Transactional);
+	MovieSceneNiagaraCacheSection->Params.SimCache = NewObject<UNiagaraSimCache>(MovieSceneNiagaraCacheSection, NAME_None, RF_Transactional);
+	MovieSceneNiagaraCacheSection->Params.CacheParameters.AttributeCaptureMode = ENiagaraSimCacheAttributeCaptureMode::RenderingOnly;
+	return MovieSceneNiagaraCacheSection;
+}
+
+void UMovieSceneNiagaraCacheTrack::RemoveAllAnimationData()
+{
+	AnimationSections.Empty();
+}
+
+bool UMovieSceneNiagaraCacheTrack::HasSection(const UMovieSceneSection& Section) const
+{
+	return AnimationSections.Contains(&Section);
+}
+
+void UMovieSceneNiagaraCacheTrack::AddSection(UMovieSceneSection& Section)
+{
+	AnimationSections.Add(&Section);
+}
+
+void UMovieSceneNiagaraCacheTrack::RemoveSection(UMovieSceneSection& Section)
+{
+	AnimationSections.Remove(&Section);
+}
+
+void UMovieSceneNiagaraCacheTrack::RemoveSectionAt(int32 SectionIndex)
+{
+	AnimationSections.RemoveAt(SectionIndex);
+}
+
+bool UMovieSceneNiagaraCacheTrack::IsEmpty() const
+{
+	return AnimationSections.Num() == 0;
+}
+
+#if WITH_EDITORONLY_DATA
+
+FText UMovieSceneNiagaraCacheTrack::GetDefaultDisplayName() const
+{
+	return LOCTEXT("TrackName", "Niagara Sim Cache");
+}
+
+#endif
+
+#undef LOCTEXT_NAMESPACE
diff --git a/Engine/Plugins/FX/NiagaraSimCaching/Source/NiagaraSimCaching/Public/Niagara/NiagaraSimCachingPlugin.h b/Engine/Plugins/FX/NiagaraSimCaching/Source/NiagaraSimCaching/Public/Niagara/NiagaraSimCachingPlugin.h
new file mode 100644
--- /dev/null
+++ b/Engine/Plugins/FX/NiagaraSimCaching/Source/NiagaraSimCaching/Public/Niagara/NiagaraSimCachingPlugin.h
@@ -0,0 +1,40 @@
+// Copyright Epic Games, Inc. All Rights Reserved.
+
+#pragma once
+
+#include "CoreMinimal.h"
+#include "Modules/ModuleInterface.h"
+#include "Modules/ModuleManager.h"
+#include "Logging/LogMacros.h"
+
+DECLARE_LOG_CATEGORY_EXTERN(LogNiagaraSimCaching, Verbose, All);
+
+class INiagaraSimCachingPlugin : public IModuleInterface
+{
+public:
+	virtual void StartupModule() override;
+	virtual void ShutdownModule() override;
+
+	/**
+	 * Singleton-like access to this module's interface.  This is just for convenience!
+	 * Beware of calling this during the shutdown phase, though.  Your module might have been unloaded already.
+	 *
+	 * @return Returns singleton instance, loading the module on demand if needed
+	 */
+	static INiagaraSimCachingPlugin& Get()
+	{
+		return FModuleManager::LoadModuleChecked<INiagaraSimCachingPlugin>("NiagaraSimCaching");
+	}
+
+	/**
+	 * Checks to see if this module is loaded and ready.  It is only valid to call Get() if IsAvailable() returns true.
+	 *
+	 * @return True if the module is loaded and ready to use
+	 */
+	static bool IsAvailable()
+	{
+		return FModuleManager::Get().IsModuleLoaded("NiagaraSimCaching");
+	}
+
+private:
+};
diff --git a/Engine/Plugins/FX/NiagaraSimCaching/Source/NiagaraSimCaching/Public/Niagara/Sequencer/MovieSceneNiagaraCacheSection.h b/Engine/Plugins/FX/NiagaraSimCaching/Source/NiagaraSimCaching/Public/Niagara/Sequencer/MovieSceneNiagaraCacheSection.h
new file mode 100644
--- /dev/null
+++ b/Engine/Plugins/FX/NiagaraSimCaching/Source/NiagaraSimCaching/Public/Niagara/Sequencer/MovieSceneNiagaraCacheSection.h
@@ -0,0 +1,40 @@
+﻿// Copyright Epic Games, Inc. All Rights Reserved.
+
+#pragma once
+
+#include "NiagaraSimCache.h"
+#include "Sections/MovieSceneBaseCacheSection.h"
+#include "MovieSceneNiagaraCacheSection.generated.h"
+
+USTRUCT()
+struct NIAGARASIMCACHING_API FMovieSceneNiagaraCacheParams : public FMovieSceneBaseCacheParams
+{
+	GENERATED_BODY()
+
+	FMovieSceneNiagaraCacheParams();
+	virtual ~FMovieSceneNiagaraCacheParams() override {}
+
+	/** Gets the animation sequence length, not modified by play rate */
+	virtual float GetSequenceLength() const override;
+
+	UPROPERTY(VisibleAnywhere, Category = "NiagaraCache")
+	FNiagaraSimCacheCreateParameters CacheParameters;
+	
+	/** The sim cache this section plays and records into */
+	UPROPERTY()
+	TObjectPtr<UNiagaraSimCache> SimCache;
+};
+
+/**
+ * Movie scene section that control NiagaraCache playback
+ */
+UCLASS(MinimalAPI)
+class UMovieSceneNiagaraCacheSection : public UMovieSceneBaseCacheSection
+{
+	GENERATED_UCLASS_BODY()
+
+public:
+
+	UPROPERTY(EditAnywhere, Category = "NiagaraCache", meta = (ShowOnlyInnerProperties))
+	FMovieSceneNiagaraCacheParams Params;
+};
diff --git a/Engine/Plugins/FX/NiagaraSimCaching/Source/NiagaraSimCaching/Public/Niagara/Sequencer/MovieSceneNiagaraCacheTrack.h b/Engine/Plugins/FX/NiagaraSimCaching/Source/NiagaraSimCaching/Public/Niagara/Sequencer/MovieSceneNiagaraCacheTrack.h
new file mode 100644
--- /dev/null
+++ b/Engine/Plugins/FX/NiagaraSimCaching/Source/NiagaraSimCaching/Public/Niagara/Sequencer/MovieSceneNiagaraCacheTrack.h
@@ -0,0 +1,50 @@
+﻿// Copyright Epic Games, Inc. All Rights Reserved.
+
+#pragma once
+
+#include "MovieSceneNameableTrack.h"
+#include "Compilation/IMovieSceneTrackTemplateProducer.h"
+#include "MovieSceneNiagaraCacheTrack.generated.h"
+
+/**
+ * Handles animation of NiagaraCache
+ */
+UCLASS(MinimalAPI)
+class UMovieSceneNiagaraCacheTrack
+	: public UMovieSceneNameableTrack
+	, public IMovieSceneTrackTemplateProducer
+{
+	GENERATED_UCLASS_BODY()
+
+public:
+
+	/** Adds a new animation to this track */
+	virtual UMovieSceneSection* AddNewAnimation(FFrameNumber KeyTime, class UNiagaraComponent* NiagaraComponent);
+
+	/** Gets the animation sections at a certain time */
+	TArray<UMovieSceneSection*> GetAnimSectionsAtTime(FFrameNumber Time);
+
+	// UMovieSceneTrack interface
+	virtual void RemoveAllAnimationData() override;
+	virtual bool HasSection(const UMovieSceneSection& Section) const override;
+	virtual void AddSection(UMovieSceneSection& Section) override;
+	virtual void RemoveSection(UMovieSceneSection& Section) override;
+	virtual void RemoveSectionAt(int32 SectionIndex) override;
+	virtual bool IsEmpty() const override;
+	virtual const TArray<UMovieSceneSection*>& GetAllSections() const override;
+	virtual bool SupportsType(TSubclassOf<UMovieSceneSection> SectionClass) const override;
+	virtual UMovieSceneSection* CreateNewSection() override;
+
+	// ~IMovieSceneTrackTemplateProducer interface
+	virtual FMovieSceneEvalTemplatePtr CreateTemplateForSection(const UMovieSceneSection& InSection) const override;
+
+#if WITH_EDITORONLY_DATA
+	virtual FText GetDefaultDisplayName() const override;
+#endif
+
+private:
+
+	/** List of all animation sections */
+	UPROPERTY()
+	TArray<TObjectPtr<UMovieSceneSection>> AnimationSections;
+};
diff --git a/Engine/Plugins/FX/NiagaraSimCaching/Source/NiagaraSimCachingEditor/NiagaraSimCachingEditor.Build.cs b/Engine/Plugins/FX/NiagaraSimCaching/Source/NiagaraSimCachingEditor/NiagaraSimCachingEditor.Build.cs
new file mode 100644
--- /dev/null
+++ b/Engine/Plugins/FX/NiagaraSimCaching/Source/NiagaraSimCachingEditor/NiagaraSimCachingEditor.Build.cs
@@ -0,0 +1,39 @@
+// Copyright Epic Games, Inc. All Rights Reserved.
+
+namespace UnrealBuildTool.Rules
+{
+	public class NiagaraSimCachingEditor : ModuleRules
+	{
+        public NiagaraSimCachingEditor(ReadOnlyTargetRules Target) : base(Target)
+		{
+			PrivateIncludePaths.Add("NiagaraSimCachingEditor/Private");
+            PublicIncludePaths.Add(ModuleDirectory + "/Public");
+
+			PublicDependencyModuleNames.AddRange(
+				new[]
+				{
+					"Core",
+					"CoreUObject",
+					"Slate",
+					"SlateCore",
+					"InputCore",
+					"Engine",
+					"UnrealEd",
+					"PropertyEditor",
+					"ToolMenus",
+					"Niagara",
+					"NiagaraSimCaching",
+					"LevelSequence",
+                    "MovieScene",
+                    "MovieSceneTools",
+                    "MovieSceneTracks",
+                    "Sequencer",
+                    "TimeManagement",
+					"EditorFramework",
+					"TakesCore",
+					"TakeRecorder",
+					"TakeTrackRecorders",
+				});
+		}
+	}
+}
diff --git a/Engine/Plugins/FX/NiagaraSimCaching/Source/NiagaraSimCachingEditor/Private/Niagara/NiagaraSimCachingEditorPlugin.cpp b/Engine/Plugins/FX/NiagaraSimCaching/Source/NiagaraSimCachingEditor/Private/Niagara/NiagaraSimCachingEditorPlugin.cpp
new file mode 100644
--- /dev/null
+++ b/Engine/Plugins/FX/NiagaraSimCaching/Source/NiagaraSimCachingEditor/Private/Niagara/NiagaraSimCachingEditorPlugin.cpp
@@ -0,0 +1,39 @@
+// Copyright Epic Games, Inc. All Rights Reserved.
+
+#include "Niagara/NiagaraSimCachingEditorPlugin.h"
+#include "Sequencer/TakeRecorderNiagaraCacheSource.h"
+#include "CoreMinimal.h"
+#include "Engine/Selection.h"
+#include "Modules/ModuleManager.h"
+#include "ISequencerModule.h"
+#include "Sequencer/NiagaraCacheTrackEditor.h"
+#include "ITakeRecorderModule.h"
+
+IMPLEMENT_MODULE(INiagaraSimCachingEditorPlugin, NiagaraSimCachingEditor)
+
+#define LOCTEXT_NAMESPACE "CacheEditorPlugin"
+
+static const FName MovieSceneTrackRecorderFactoryName("MovieSceneTrackRecorderFactory");
+
+void INiagaraSimCachingEditorPlugin::StartupModule()
+{	
+	ISequencerModule& SequencerModule = FModuleManager::Get().LoadModuleChecked<ISequencerModule>("Sequencer");
+	TrackEditorBindingHandle = SequencerModule.RegisterTrackEditor(FOnCreateTrackEditor::CreateStatic(&FNiagaraCacheTrackEditor::CreateTrackEditor));
+		
+	IModularFeatures::Get().RegisterModularFeature(MovieSceneTrackRecorderFactoryName, &MovieSceneNiagaraCacheTrackRecorder);
+}
+
+void INiagaraSimCachingEditorPlugin::ShutdownModule()
+{
+	if(UObjectInitialized())
+	{
+		if (ISequencerModule* SequencerModulePtr = FModuleManager::Get().GetModulePtr<ISequencerModule>("Sequencer"))
+		{
+			SequencerModulePtr->UnRegisterTrackEditor(TrackEditorBindingHandle);
+		}
+
+		IModularFeatures::Get().UnregisterModularFeature(MovieSceneTrackRecorderFactoryName, &MovieSceneNiagaraCacheTrackRecorder);
+	}
+}
+
+#undef LOCTEXT_NAMESPACE
diff --git a/Engine/Plugins/FX/NiagaraSimCaching/Source/NiagaraSimCachingEditor/Private/Niagara/NiagaraSimCachingEditorStyle.cpp b/Engine/Plugins/FX/NiagaraSimCaching/Source/NiagaraSimCachingEditor/Private/Niagara/NiagaraSimCachingEditorStyle.cpp
new file mode 100644
--- /dev/null
+++ b/Engine/Plugins/FX/NiagaraSimCaching/Source/NiagaraSimCachingEditor/Private/Niagara/NiagaraSimCachingEditorStyle.cpp
@@ -0,0 +1,31 @@
+// Copyright Epic Games, Inc. All Rights Reserved.
+
+#include "NiagaraSimCachingEditorStyle.h"
+
+#include "Styling/SlateStyleRegistry.h"
+
+FName FNiagaraSimCachingEditorStyle::StyleName("NiagaraSimCachingEditorStyle");
+
+FNiagaraSimCachingEditorStyle::FNiagaraSimCachingEditorStyle()
+	: FSlateStyleSet(StyleName)
+{
+	const FVector2D IconSize(20.0f, 20.0f);
+	const FVector2D LabelIconSize(16.0f, 16.0f);
+
+	FSlateStyleSet::SetCoreContentRoot(FPaths::EngineContentDir() / TEXT("Editor/Slate"));
+
+	FSlateStyleRegistry::RegisterSlateStyle(*this);
+}
+
+FNiagaraSimCachingEditorStyle::~FNiagaraSimCachingEditorStyle()
+{
+	FSlateStyleRegistry::UnRegisterSlateStyle(*this);
+}
+
+FNiagaraSimCachingEditorStyle& FNiagaraSimCachingEditorStyle::Get()
+{
+	static FNiagaraSimCachingEditorStyle Inst;
+	return Inst;
+}
+
+
diff --git a/Engine/Plugins/FX/NiagaraSimCaching/Source/NiagaraSimCachingEditor/Private/Niagara/NiagaraSimCachingEditorStyle.h b/Engine/Plugins/FX/NiagaraSimCaching/Source/NiagaraSimCachingEditor/Private/Niagara/NiagaraSimCachingEditorStyle.h
new file mode 100644
--- /dev/null
+++ b/Engine/Plugins/FX/NiagaraSimCaching/Source/NiagaraSimCachingEditor/Private/Niagara/NiagaraSimCachingEditorStyle.h
@@ -0,0 +1,23 @@
+// Copyright Epic Games, Inc. All Rights Reserved.
+
+#pragma once
+
+#include "Styling/SlateStyle.h"
+
+/**
+ * Slate style set that defines all the styles for the take recorder UI
+ */
+class FNiagaraSimCachingEditorStyle
+	: public FSlateStyleSet
+{
+public:
+	static FName StyleName;
+
+	/** Access the singleton instance for this style set */
+	static FNiagaraSimCachingEditorStyle& Get();
+
+private:
+
+	FNiagaraSimCachingEditorStyle();
+	virtual ~FNiagaraSimCachingEditorStyle() override;
+};
diff --git a/Engine/Plugins/FX/NiagaraSimCaching/Source/NiagaraSimCachingEditor/Private/Sequencer/MovieSceneNiagaraTrackRecorder.cpp b/Engine/Plugins/FX/NiagaraSimCaching/Source/NiagaraSimCachingEditor/Private/Sequencer/MovieSceneNiagaraTrackRecorder.cpp
new file mode 100644
--- /dev/null
+++ b/Engine/Plugins/FX/NiagaraSimCaching/Source/NiagaraSimCachingEditor/Private/Sequencer/MovieSceneNiagaraTrackRecorder.cpp
@@ -0,0 +1,126 @@
+﻿// Copyright Epic Games, Inc. All Rights Reserved.
+
+#include "Sequencer/MovieSceneNiagaraTrackRecorder.h"
+#include "Niagara/Sequencer/MovieSceneNiagaraCacheSection.h"
+#include "Niagara/Sequencer/MovieSceneNiagaraCacheTrack.h"
+#include "TakeRecorderSettings.h"
+#include "Recorder/TakeRecorderParameters.h"
+#include "Engine/TimecodeProvider.h"
+#include "Misc/App.h"
+#include "MovieSceneFolder.h"
+#include "LevelSequence.h"
+#include "NiagaraComponent.h"
+#include "NiagaraSimCacheFunctionLibrary.h"
+
+#define LOCTEXT_NAMESPACE "MovieSceneNiagaraTrackRecorder"
+
+bool FMovieSceneNiagaraTrackRecorderFactory::CanRecordObject(UObject* InObjectToRecord) const
+{
+	return InObjectToRecord->IsA<UNiagaraComponent>();
+}
+
+UMovieSceneTrackRecorder* FMovieSceneNiagaraTrackRecorderFactory::CreateTrackRecorderForObject() const
+{
+	return NewObject<UMovieSceneNiagaraTrackRecorder>();
+}
+
+UMovieSceneTrackRecorder* FMovieSceneNiagaraTrackRecorderFactory::CreateTrackRecorderForProperty(UObject* InObjectToRecord, const FName& InPropertyToRecord) const
+{
+	return nullptr;
+}
+
+void UMovieSceneNiagaraTrackRecorder::CreateTrackImpl()
+{
+	SystemToRecord = CastChecked<UNiagaraComponent>(ObjectToRecord.Get());
+
+	NiagaraCacheTrack = MovieScene->FindTrack<UMovieSceneNiagaraCacheTrack>(ObjectGuid);
+	if (!NiagaraCacheTrack.IsValid())
+	{
+		NiagaraCacheTrack = MovieScene->AddTrack<UMovieSceneNiagaraCacheTrack>(ObjectGuid);
+	}
+	else
+	{
+		NiagaraCacheTrack->RemoveAllAnimationData();
+	}
+
+	if (NiagaraCacheTrack.IsValid())
+	{
+		NiagaraCacheSection = CastChecked<UMovieSceneNiagaraCacheSection>(NiagaraCacheTrack->CreateNewSection());
+		NiagaraCacheSection->SetIsActive(false);
+		NiagaraCacheTrack->AddSection(*NiagaraCacheSection);
+
+		// Resize the section to either it's remaining	keyframes range or 0
+		NiagaraCacheSection->SetRange(NiagaraCacheSection->GetAutoSizeRange().Get(TRange<FFrameNumber>(0, 0)));
+
+		// Make sure it starts at frame 0, in case Auto Size removed a piece of the start
+		NiagaraCacheSection->ExpandToFrame(0);
+	}
+}
+
+void UMovieSceneNiagaraTrackRecorder::SetSectionStartTimecodeImpl(const FTimecode& InSectionStartTimecode, const FFrameNumber& InSectionFirstFrame)
+{
+	if (NiagaraCacheSection.IsValid() && NiagaraCacheTrack.IsValid())
+	{
+		NiagaraCacheSection->TimecodeSource = FMovieSceneTimecodeSource(InSectionStartTimecode);
+
+		FTakeRecorderParameters Parameters;
+		Parameters.User = GetDefault<UTakeRecorderUserSettings>()->Settings;
+		Parameters.Project = GetDefault<UTakeRecorderProjectSettings>()->Settings;
+
+		const FFrameRate TickResolution = MovieScene->GetTickResolution();
+		const FFrameRate DisplayRate = MovieScene->GetDisplayRate();
+
+		RecordStartTime = FApp::GetCurrentTime();
+		RecordStartFrame = Parameters.Project.bStartAtCurrentTimecode ?
+			FFrameRate::TransformTime(FFrameTime(InSectionStartTimecode.ToFrameNumber(DisplayRate)),
+				DisplayRate, TickResolution).FloorToFrame() : MovieScene->GetPlaybackRange().GetLowerBoundValue();
+
+		if (SystemToRecord.IsValid())
+		{
+			// start simulation and writing to the sim cache 
+			SystemToRecord->SetSimCache(nullptr);
+			//SystemToRecord->SetAgeUpdateMode(ENiagaraAgeUpdateMode::TickDeltaTime);
+			//SystemToRecord->ResetSystem();
+			NiagaraCacheSection->Params.SimCache->BeginWrite(NiagaraCacheSection->Params.CacheParameters, SystemToRecord.Get());
+		}
+	}
+}
+
+UMovieSceneSection* UMovieSceneNiagaraTrackRecorder::GetMovieSceneSection() const
+{
+	return Cast<UMovieSceneSection>(NiagaraCacheSection.Get());
+}
+
+void UMovieSceneNiagaraTrackRecorder::FinalizeTrackImpl()
+{
+	if (NiagaraCacheSection.IsValid())
+	{
+		// finalize the sim cache
+		NiagaraCacheSection->Params.SimCache->EndWrite();
+		
+		// Set the final range 
+		TOptional<TRange<FFrameNumber> > DefaultSectionLength = NiagaraCacheSection->GetAutoSizeRange();
+		if (DefaultSectionLength.IsSet())
+		{
+			NiagaraCacheSection->SetRange(DefaultSectionLength.GetValue());
+		}
+
+		// Activate the section
+		NiagaraCacheSection->SetIsActive(true);
+	}
+}
+
+void UMovieSceneNiagaraTrackRecorder::RecordSampleImpl(const FQualifiedFrameTime& CurrentFrameTime)
+{
+	// Expand the section to the new length
+	const FFrameRate	 TickResolution = MovieScene->GetTickResolution();
+	const FFrameNumber CurrentFrame = CurrentFrameTime.ConvertTo(TickResolution).FloorToFrame();
+
+	if (NiagaraCacheSection.IsValid() && SystemToRecord.IsValid())
+	{
+		NiagaraCacheSection->Params.SimCache->WriteFrame(SystemToRecord.Get());
+		NiagaraCacheSection->SetEndFrame(CurrentFrame);
+	}
+}
+
+#undef LOCTEXT_NAMESPACE
diff --git a/Engine/Plugins/FX/NiagaraSimCaching/Source/NiagaraSimCachingEditor/Private/Sequencer/NiagaraCacheTrackEditor.cpp b/Engine/Plugins/FX/NiagaraSimCaching/Source/NiagaraSimCachingEditor/Private/Sequencer/NiagaraCacheTrackEditor.cpp
new file mode 100644
--- /dev/null
+++ b/Engine/Plugins/FX/NiagaraSimCaching/Source/NiagaraSimCachingEditor/Private/Sequencer/NiagaraCacheTrackEditor.cpp
@@ -0,0 +1,347 @@
+﻿// Copyright Epic Games, Inc. All Rights Reserved.
+
+#include "Sequencer/NiagaraCacheTrackEditor.h"
+#include "CommonMovieSceneTools.h"
+#include "Fonts/FontMeasure.h"
+#include "Framework/MultiBox/MultiBoxBuilder.h"
+#include "LevelSequence.h"
+#include "NiagaraComponent.h"
+#include "NiagaraSimCache.h"
+#include "Niagara/Sequencer/MovieSceneNiagaraCacheTrack.h"
+#include "Niagara/Sequencer/MovieSceneNiagaraCacheSection.h"
+#include "SequencerSectionPainter.h"
+#include "SequencerUtilities.h"
+#include "Styling/SlateIconFinder.h"
+#include "TimeToPixel.h"
+
+namespace NiagaraCacheEditorConstants
+{
+	constexpr float AnimationTrackHeight = 20.f;
+}
+
+#define LOCTEXT_NAMESPACE "FNiagaraCacheTrackEditor"
+
+static UNiagaraComponent* AcquireNiagaraComponentFromObjectGuid(const FGuid& Guid, TSharedPtr<ISequencer> SequencerPtr)
+{
+	UObject* BoundObject = SequencerPtr.IsValid() ? SequencerPtr->FindSpawnedObjectOrTemplate(Guid) : nullptr;
+	return Cast<UNiagaraComponent>(BoundObject);
+}
+
+FNiagaraCacheSection::FNiagaraCacheSection(UMovieSceneSection& InSection, TWeakPtr<ISequencer> InSequencer)
+	: Section(*CastChecked<UMovieSceneNiagaraCacheSection>(&InSection))
+	, Sequencer(InSequencer)
+	, InitialFirstLoopStartOffsetDuringResize(0)
+	, InitialStartTimeDuringResize(0)
+{ }
+
+UMovieSceneSection* FNiagaraCacheSection::GetSectionObject()
+{
+	return &Section;
+}
+
+FText FNiagaraCacheSection::GetSectionTitle() const
+{
+	if (Section.Params.SimCache != nullptr)
+	{
+		FNumberFormattingOptions FormatOptions;
+		FormatOptions.MaximumFractionalDigits = 1;
+		return FText::Format(FText::FromString("Sim Cache ({0} frames/{1}s)"), Section.Params.SimCache->GetNumFrames(), FText::AsNumber(Section.Params.SimCache->GetDurationSeconds(), &FormatOptions));
+	
+	}
+	return LOCTEXT("NoNiagaraCacheSection", "No NiagaraCache");
+}
+
+float FNiagaraCacheSection::GetSectionHeight() const
+{
+	return NiagaraCacheEditorConstants::AnimationTrackHeight;
+}
+
+int32 FNiagaraCacheSection::OnPaintSection(FSequencerSectionPainter& Painter) const
+{
+	using namespace UE::Sequencer;
+
+	const ESlateDrawEffect DrawEffects = Painter.bParentEnabled ? ESlateDrawEffect::None : ESlateDrawEffect::DisabledEffect;
+
+	const FTimeToPixel& TimeToPixelConverter = Painter.GetTimeConverter();
+
+	int32 LayerId = Painter.PaintSectionBackground();
+
+	static const FSlateBrush* GenericDivider = FAppStyle::GetBrush("Sequencer.GenericDivider");
+
+	if (!Section.HasStartFrame() || !Section.HasEndFrame())
+	{
+		return LayerId;
+	}
+
+	FFrameRate TickResolution = TimeToPixelConverter.GetTickResolution();
+
+	// Add lines where the animation starts and ends/loops
+	const float AnimPlayRate = FMath::IsNearlyZero(Section.Params.PlayRate) ? 1.0f : Section.Params.PlayRate;
+	const float Duration = Section.Params.GetSequenceLength();
+	const float SeqLength = Duration - TickResolution.AsSeconds(Section.Params.StartFrameOffset + Section.Params.EndFrameOffset) / AnimPlayRate;
+	const float FirstLoopSeqLength = SeqLength - TickResolution.AsSeconds(Section.Params.FirstLoopStartFrameOffset) / AnimPlayRate;
+
+	if (!FMath::IsNearlyZero(SeqLength, KINDA_SMALL_NUMBER) && SeqLength > 0)
+	{
+		float MaxOffset = Section.GetRange().Size<FFrameTime>() / TickResolution;
+		float OffsetTime = FirstLoopSeqLength;
+		float StartTime = Section.GetInclusiveStartFrame() / TickResolution;
+
+		while (OffsetTime < MaxOffset)
+		{
+			float OffsetPixel = TimeToPixelConverter.SecondsToPixel(StartTime + OffsetTime) - TimeToPixelConverter.SecondsToPixel(StartTime);
+
+			FSlateDrawElement::MakeBox(
+				Painter.DrawElements,
+				LayerId,
+				Painter.SectionGeometry.MakeChild(
+					FVector2D(2.f, Painter.SectionGeometry.Size.Y - 2.f),
+					FSlateLayoutTransform(FVector2D(OffsetPixel, 1.f))
+				).ToPaintGeometry(),
+				GenericDivider,
+				DrawEffects
+			);
+
+			OffsetTime += SeqLength;
+		}
+	}
+
+	TSharedPtr<ISequencer> SequencerPtr = Sequencer.Pin();
+	if (Painter.bIsSelected && SequencerPtr.IsValid())
+	{
+		FFrameTime CurrentTime = SequencerPtr->GetLocalTime().Time;
+		if (Section.GetRange().Contains(CurrentTime.FrameNumber) && Section.Params.SimCache != nullptr)
+		{
+			const float Time = TimeToPixelConverter.FrameToPixel(CurrentTime);
+
+			UNiagaraSimCache* NiagaraCache = Section.Params.SimCache;
+			int32 NumFrames = NiagaraCache ? NiagaraCache->GetNumFrames() : 0;
+
+			// Draw the current time next to the scrub handle
+			const float AnimTime = Section.MapTimeToAnimation(Duration, CurrentTime, TickResolution);
+			int32 FrameTime = Duration ? FMath::FloorToFloat((AnimTime / Duration) * NumFrames) + 1 : 0;
+			FString FrameString =  FString::FromInt(FrameTime);
+
+			const FSlateFontInfo SmallLayoutFont = FCoreStyle::GetDefaultFontStyle("Bold", 10);
+			const TSharedRef< FSlateFontMeasure > FontMeasureService = FSlateApplication::Get().GetRenderer()->GetFontMeasureService();
+			FVector2D TextSize = FontMeasureService->Measure(FrameString, SmallLayoutFont);
+
+			// Flip the text position if getting near the end of the view range
+			static constexpr float TextOffsetPx = 10.f;
+			bool  bDrawLeft = (Painter.SectionGeometry.Size.X - Time) < (TextSize.X + 22.f) - TextOffsetPx;
+			float TextPosition = bDrawLeft ? Time - TextSize.X - TextOffsetPx : Time + TextOffsetPx;
+			//handle mirrored labels
+			constexpr float MajorTickHeight = 7.0f;
+			FVector2D TextOffset(TextPosition, Painter.SectionGeometry.Size.Y - (MajorTickHeight + TextSize.Y));
+
+			const FLinearColor DrawColor = FAppStyle::GetSlateColor("SelectionColor").GetColor(FWidgetStyle());
+			const FVector2D BoxPadding = FVector2D(4.0f, 2.0f);
+			// draw time string
+
+			FSlateDrawElement::MakeBox(
+				Painter.DrawElements,
+				LayerId + 5,
+				Painter.SectionGeometry.ToPaintGeometry(TextOffset - BoxPadding, TextSize + 2.0f * BoxPadding),
+				FAppStyle::GetBrush("WhiteBrush"),
+				ESlateDrawEffect::None,
+				FLinearColor::Black.CopyWithNewOpacity(0.5f)
+			);
+
+			FSlateDrawElement::MakeText(
+				Painter.DrawElements,
+				LayerId + 6,
+				Painter.SectionGeometry.ToPaintGeometry(TextOffset, TextSize),
+				FrameString,
+				SmallLayoutFont,
+				DrawEffects,
+				DrawColor
+			);
+
+		}
+	}
+
+	return LayerId;
+}
+
+void FNiagaraCacheSection::BeginResizeSection()
+{
+	InitialFirstLoopStartOffsetDuringResize = Section.Params.FirstLoopStartFrameOffset;
+	InitialStartTimeDuringResize = Section.HasStartFrame() ? Section.GetInclusiveStartFrame() : 0;
+}
+
+void FNiagaraCacheSection::UpdateSection(FFrameNumber& UpdateTime) const
+{
+	const FFrameRate FrameRate = Section.GetTypedOuter<UMovieScene>()->GetTickResolution();
+	FFrameNumber StartOffset = FrameRate.AsFrameNumber((UpdateTime - InitialStartTimeDuringResize) / FrameRate * Section.Params.PlayRate);
+
+	StartOffset += InitialFirstLoopStartOffsetDuringResize;
+
+	if (StartOffset < 0)
+	{
+		// Ensure start offset is not less than 0 and adjust ResizeTime
+		UpdateTime = UpdateTime - StartOffset;
+
+		StartOffset = FFrameNumber(0);
+	}
+	else
+	{
+		// If the start offset exceeds the length of one loop, trim it back.
+		const FFrameNumber SeqLength = FrameRate.AsFrameNumber(Section.Params.GetSequenceLength()) - Section.Params.StartFrameOffset - Section.Params.EndFrameOffset;
+		StartOffset = StartOffset % SeqLength;
+	}
+
+	Section.Params.FirstLoopStartFrameOffset = StartOffset;
+}
+
+void FNiagaraCacheSection::ResizeSection(ESequencerSectionResizeMode ResizeMode, FFrameNumber ResizeTime)
+{
+	// Adjust the start offset when resizing from the beginning
+	if (ResizeMode == SSRM_LeadingEdge)
+	{
+		UpdateSection(ResizeTime);
+	}
+
+	ISequencerSection::ResizeSection(ResizeMode, ResizeTime);
+}
+
+void FNiagaraCacheSection::BeginSlipSection()
+{
+	BeginResizeSection();
+}
+
+void FNiagaraCacheSection::SlipSection(FFrameNumber SlipTime)
+{
+	UpdateSection(SlipTime);
+
+	ISequencerSection::SlipSection(SlipTime);
+}
+
+void FNiagaraCacheSection::BeginDilateSection()
+{
+	Section.PreviousPlayRate = Section.Params.PlayRate; //make sure to cache the play rate
+}
+void FNiagaraCacheSection::DilateSection(const TRange<FFrameNumber>& NewRange, float DilationFactor)
+{
+	Section.Params.PlayRate = Section.PreviousPlayRate / DilationFactor;
+	Section.SetRange(NewRange);
+}
+
+FNiagaraCacheTrackEditor::FNiagaraCacheTrackEditor(TSharedRef<ISequencer> InSequencer)
+	: FMovieSceneTrackEditor(InSequencer)
+{ }
+
+TSharedRef<ISequencerTrackEditor> FNiagaraCacheTrackEditor::CreateTrackEditor(TSharedRef<ISequencer> InSequencer)
+{
+	return MakeShareable(new FNiagaraCacheTrackEditor(InSequencer));
+}
+
+bool FNiagaraCacheTrackEditor::SupportsSequence(UMovieSceneSequence* InSequence) const
+{
+	if (InSequence && InSequence->IsTrackSupported(UMovieSceneNiagaraCacheTrack::StaticClass()) == ETrackSupport::NotSupported)
+	{
+		return false;
+	}
+
+	return InSequence && InSequence->IsA(ULevelSequence::StaticClass());
+}
+
+bool FNiagaraCacheTrackEditor::SupportsType(TSubclassOf<UMovieSceneTrack> Type) const
+{
+	return Type == UMovieSceneNiagaraCacheTrack::StaticClass();
+}
+
+TSharedRef<ISequencerSection> FNiagaraCacheTrackEditor::MakeSectionInterface(UMovieSceneSection& SectionObject, UMovieSceneTrack& Track, FGuid ObjectBinding)
+{
+	check(SupportsType(SectionObject.GetOuter()->GetClass()));
+
+	return MakeShareable(new FNiagaraCacheSection(SectionObject, GetSequencer()));
+}
+
+void FNiagaraCacheTrackEditor::BuildObjectBindingTrackMenu(FMenuBuilder& MenuBuilder, const TArray<FGuid>& ObjectBindings, const UClass* ObjectClass)
+{
+	//TODO (mga) allow users to create cache tracks directly?
+	/*
+	if (ObjectClass->IsChildOf(UNiagaraComponent::StaticClass()) || ObjectClass->IsChildOf(AActor::StaticClass()))
+	{
+		if(ObjectBindings.Num() > 0)
+		{
+			if (UNiagaraComponent* NiagaraComponent = AcquireNiagaraComponentFromObjectGuid(ObjectBindings[0], GetSequencer()))
+			{
+				UMovieSceneTrack* Track = nullptr;
+
+				
+				MenuBuilder.AddMenuEntry(
+					NSLOCTEXT("Sequencer", "AddNiagaraCache", "Niagara Cache"),
+					NSLOCTEXT("Sequencer", "AddNiagaraCacheTooltip", "Adds a Niagara Cache track."),
+					FSlateIcon(),
+					FUIAction(
+						FExecuteAction::CreateSP(this, &FNiagaraCacheTrackEditor::BuildNiagaraCacheTrack, ObjectBindings, Track)
+					)
+				);
+				
+			}
+		}
+	}*/
+}
+
+void FNiagaraCacheTrackEditor::BuildNiagaraCacheTrack(TArray<FGuid> ObjectBindings, UMovieSceneTrack* Track)
+{
+	const TSharedPtr<ISequencer> SequencerPtr = GetSequencer();
+
+	if (SequencerPtr.IsValid())
+	{
+		const FScopedTransaction Transaction(LOCTEXT("AddNiagaraCache_Transaction", "Add Niagara Cache"));
+
+		for (FGuid ObjectBinding : ObjectBindings)
+		{
+			if (UNiagaraComponent* NiagaraComponent = AcquireNiagaraComponentFromObjectGuid(ObjectBinding, SequencerPtr))
+				{
+					AnimatablePropertyChanged(FOnKeyProperty::CreateRaw(this, &FNiagaraCacheTrackEditor::AddKeyInternal, NiagaraComponent, Track));
+				}
+		}
+	}
+}
+
+FKeyPropertyResult FNiagaraCacheTrackEditor::AddKeyInternal(FFrameNumber KeyTime, UNiagaraComponent* NiagaraComponent, UMovieSceneTrack* Track)
+{
+	FKeyPropertyResult KeyPropertyResult;
+
+	const FFindOrCreateHandleResult HandleResult = FindOrCreateHandleToObject(NiagaraComponent);
+	const FGuid ObjectHandle = HandleResult.Handle;
+	KeyPropertyResult.bHandleCreated |= HandleResult.bWasCreated;
+	if (ObjectHandle.IsValid())
+	{
+		if (!Track)
+		{
+			Track = AddTrack(GetSequencer()->GetFocusedMovieSceneSequence()->GetMovieScene(), ObjectHandle, UMovieSceneNiagaraCacheTrack::StaticClass(), NAME_None);
+			KeyPropertyResult.bTrackCreated = true;
+		}
+
+		if (ensure(Track))
+		{
+			Track->Modify();
+
+			UMovieSceneSection* NewSection = Cast<UMovieSceneNiagaraCacheTrack>(Track)->AddNewAnimation(KeyTime, NiagaraComponent);
+			KeyPropertyResult.bTrackModified = true;
+			KeyPropertyResult.SectionsCreated.Add(NewSection);
+
+			GetSequencer()->EmptySelection();
+			GetSequencer()->SelectSection(NewSection);
+			GetSequencer()->ThrobSectionSelection();
+		}
+	}
+
+	return KeyPropertyResult;
+}
+
+TSharedPtr<SWidget> FNiagaraCacheTrackEditor::BuildOutlinerEditWidget(const FGuid& ObjectBinding, UMovieSceneTrack* Track, const FBuildEditWidgetParams& Params)
+{
+	return TSharedPtr<SWidget>();
+}
+
+const FSlateBrush* FNiagaraCacheTrackEditor::GetIconBrush() const
+{
+	return FSlateIconFinder::FindIconForClass(UNiagaraSimCache::StaticClass()).GetIcon();
+}
+
+#undef LOCTEXT_NAMESPACE
diff --git a/Engine/Plugins/FX/NiagaraSimCaching/Source/NiagaraSimCachingEditor/Private/Sequencer/TakeRecorderNiagaraCacheSource.cpp b/Engine/Plugins/FX/NiagaraSimCaching/Source/NiagaraSimCachingEditor/Private/Sequencer/TakeRecorderNiagaraCacheSource.cpp
new file mode 100644
--- /dev/null
+++ b/Engine/Plugins/FX/NiagaraSimCaching/Source/NiagaraSimCachingEditor/Private/Sequencer/TakeRecorderNiagaraCacheSource.cpp
@@ -0,0 +1,112 @@
+﻿// Copyright Epic Games, Inc. All Rights Reserved.
+
+#include "Sequencer/TakeRecorderNiagaraCacheSource.h"
+
+#include "Niagara/Sequencer/MovieSceneNiagaraCacheTrack.h"
+#include "Sequencer/MovieSceneNiagaraTrackRecorder.h"
+
+#include "LevelSequence.h"
+#include "MovieSceneFolder.h"
+#include "NiagaraComponent.h"
+#include "Channels/MovieSceneFloatChannel.h"
+#include "Channels/MovieSceneChannelTraits.h"
+#include "Framework/Notifications/NotificationManager.h" 
+
+
+#define LOCTEXT_NAMESPACE "TakeRecorderNiagaraCacheSource"
+
+UTakeRecorderNiagaraCacheSource::UTakeRecorderNiagaraCacheSource(const FObjectInitializer& ObjInit)
+	: Super(ObjInit)
+{
+	TrackTint = FColor(0.0f, 125.0f, 255.0f, 65.0f);
+}
+
+TArray<UTakeRecorderSource*> UTakeRecorderNiagaraCacheSource::PreRecording(ULevelSequence* InSequence, FMovieSceneSequenceID InSequenceID, ULevelSequence* InMasterSequence, FManifestSerializer* InManifestSerializer)
+{
+	if (SystemToRecord)
+	{
+		UMovieScene* MovieScene = InSequence->GetMovieScene();
+		//FGuid ObjectGuid = MovieScene->AddSpawnable(SystemToRecord->GetActorLabel(), *SystemToRecord.Get());
+		//InSequence->Bind(ObjectGuid, *SystemToRecord.Get(), SystemToRecord->GetWorld());
+		TrackRecorder = NewObject<UMovieSceneNiagaraTrackRecorder>();
+
+		// We only support possessable for now since the spawnable template is clearing out the component reference
+		// Would be better to use a TSoftObjectPtr for the actor and have a mechanism similar to the FixUpPIE to go
+		// from PIE to editor to have a correct soft object path
+		//const FGuid ObjectGuid = MovieScene->AddPossessable(NiagaraCacheManager->GetActorLabel(), NiagaraCacheManager->GetClass());
+		//InMasterSequence->BindPossessableObject(ObjectGuid, *NiagaraCacheManager, NiagaraCacheManager->GetWorld());
+
+		//InMasterSequence->
+		//TrackRecorder->CreateTrack(MovieScene, SystemToRecord->GetNiagaraComponent(), MovieScene, nullptr, ObjectGuid);
+		CachedNiagaraCacheTrack = TrackRecorder->GetNiagaraCacheTrack();
+	}
+	return TArray<UTakeRecorderSource*>();
+}
+
+void UTakeRecorderNiagaraCacheSource::TickRecording(const FQualifiedFrameTime& CurrentTime)
+{
+	if (TrackRecorder)
+	{
+		TrackRecorder->RecordSample(CurrentTime);
+	}
+}
+
+void UTakeRecorderNiagaraCacheSource::StartRecording(const FTimecode& InSectionStartTimecode, const FFrameNumber& InSectionFirstFrame, class ULevelSequence* InSequence)
+{
+	if (TrackRecorder)
+	{
+		TrackRecorder->SetSectionStartTimecode(InSectionStartTimecode, InSectionFirstFrame);
+	}
+}
+
+void UTakeRecorderNiagaraCacheSource::StopRecording(class ULevelSequence* InSequence)
+{
+	if (TrackRecorder)
+	{
+		TrackRecorder->StopRecording();
+	}
+}
+
+TArray<UTakeRecorderSource*> UTakeRecorderNiagaraCacheSource::PostRecording(class ULevelSequence* InSequence, ULevelSequence* InMasterSequence, const bool bCancelled)
+{
+	if (TrackRecorder)
+	{
+		TrackRecorder->FinalizeTrack();
+	}
+
+	TrackRecorder = nullptr;
+	return TArray<UTakeRecorderSource*>();
+}
+
+void UTakeRecorderNiagaraCacheSource::AddContentsToFolder(UMovieSceneFolder* InFolder)
+{
+	if (CachedNiagaraCacheTrack.IsValid())
+	{
+		InFolder->AddChildMasterTrack(CachedNiagaraCacheTrack.Get());
+	}
+}
+
+FText UTakeRecorderNiagaraCacheSource::GetDisplayTextImpl() const
+{
+	/*
+	if (NiagaraCacheManager != nullptr && NiagaraCacheManager->IsValidLowLevelFast())
+	{
+		return FText::FromString(NiagaraCacheManager->GetName());
+	}*/
+
+	return LOCTEXT("Display Text", "Niagara Cache");
+}
+
+void UTakeRecorderNiagaraCacheSource::PostEditChangeProperty(FPropertyChangedEvent& PropertyChangedEvent)
+{
+	Super::PostEditChangeProperty(PropertyChangedEvent);
+
+	const FName PropertyName = PropertyChangedEvent.GetPropertyName();
+}
+
+void UTakeRecorderNiagaraCacheSource::PostLoad()
+{
+	Super::PostLoad();
+}
+
+#undef LOCTEXT_NAMESPACE
diff --git a/Engine/Plugins/FX/NiagaraSimCaching/Source/NiagaraSimCachingEditor/Public/Niagara/NiagaraSimCachingEditorPlugin.h b/Engine/Plugins/FX/NiagaraSimCaching/Source/NiagaraSimCachingEditor/Public/Niagara/NiagaraSimCachingEditorPlugin.h
new file mode 100644
--- /dev/null
+++ b/Engine/Plugins/FX/NiagaraSimCaching/Source/NiagaraSimCachingEditor/Public/Niagara/NiagaraSimCachingEditorPlugin.h
@@ -0,0 +1,51 @@
+// Copyright Epic Games, Inc. All Rights Reserved.
+
+#pragma once
+
+#include "CoreMinimal.h"
+#include "Modules/ModuleInterface.h"
+#include "Modules/ModuleManager.h"
+#include "Delegates/IDelegateInstance.h"
+#include "Sequencer/MovieSceneNiagaraTrackRecorder.h"
+
+class FAssetTypeActions_NiagaraCacheCollection;
+
+class UToolMenu;
+struct FToolMenuSection;
+class AActor;
+
+/**
+ * The public interface to this module
+ */
+class INiagaraSimCachingEditorPlugin : public IModuleInterface
+{
+	TArray<IConsoleObject*> EditorCommands;
+
+public:
+	virtual void StartupModule() override;
+	virtual void ShutdownModule() override;
+
+	/**
+	 * Singleton-like access to this module's interface.  This is just for convenience!
+	 * Beware of calling this during the shutdown phase, though.  Your module might have been unloaded already.
+	 *
+	 * @return Returns singleton instance, loading the module on demand if needed
+	 */
+	static INiagaraSimCachingEditorPlugin& Get()
+	{
+		return FModuleManager::LoadModuleChecked<INiagaraSimCachingEditorPlugin>("NiagaraSimCachingEditorPlugin");
+	}
+
+	/**
+	 * Checks to see if this module is loaded and ready.  It is only valid to call Get() if IsAvailable() returns true.
+	 *
+	 * @return True if the module is loaded and ready to use
+	 */
+	static bool IsAvailable() { return FModuleManager::Get().IsModuleLoaded("NiagaraSimCachingEditorPlugin"); }
+
+private:
+
+	FDelegateHandle TrackEditorBindingHandle;
+
+	FMovieSceneNiagaraTrackRecorderFactory MovieSceneNiagaraCacheTrackRecorder;
+};
diff --git a/Engine/Plugins/FX/NiagaraSimCaching/Source/NiagaraSimCachingEditor/Public/Sequencer/MovieSceneNiagaraTrackRecorder.h b/Engine/Plugins/FX/NiagaraSimCaching/Source/NiagaraSimCachingEditor/Public/Sequencer/MovieSceneNiagaraTrackRecorder.h
new file mode 100644
--- /dev/null
+++ b/Engine/Plugins/FX/NiagaraSimCaching/Source/NiagaraSimCachingEditor/Public/Sequencer/MovieSceneNiagaraTrackRecorder.h
@@ -0,0 +1,72 @@
+﻿// Copyright Epic Games, Inc. All Rights Reserved.
+
+#pragma once
+
+#include "TrackRecorders/MovieSceneTrackRecorder.h"
+#include "CoreMinimal.h"
+#include "MovieScene.h"
+#include "TrackRecorders/IMovieSceneTrackRecorderFactory.h"
+#include "Channels/MovieSceneFloatChannel.h"
+
+#include "MovieSceneNiagaraTrackRecorder.generated.h"
+
+class UNiagaraComponent;
+class  UMovieSceneNiagaraCacheTrack;
+class  UMovieSceneNiagaraCacheSection;
+
+struct FFrameNumber;
+class  UMovieSceneSection;
+class  UMovieSceneTrackRecorderSettings;
+
+class NIAGARASIMCACHINGEDITOR_API FMovieSceneNiagaraTrackRecorderFactory : public IMovieSceneTrackRecorderFactory
+{
+public:
+	virtual ~FMovieSceneNiagaraTrackRecorderFactory() {}
+
+	// ~Begin IMovieSceneTrackRecorderFactory Interface
+	virtual bool CanRecordObject(class UObject* InObjectToRecord) const override;
+	virtual UMovieSceneTrackRecorder* CreateTrackRecorderForObject() const override;
+
+	// Particle Systems are entire components and you can't animate them as a property
+	virtual bool CanRecordProperty(class UObject* InObjectToRecord, class FProperty* InPropertyToRecord) const override { return false; }
+	virtual UMovieSceneTrackRecorder* CreateTrackRecorderForProperty(UObject* InObjectToRecord, const FName& InPropertyToRecord) const override;
+	virtual FText GetDisplayName() const override { return NSLOCTEXT("MovieSceneNiagaraTrackRecorderFactory", "DisplayName", "Niagara Cache Track"); }
+	// ~End IMovieSceneTrackRecorderFactory Interface
+};
+
+UCLASS(BlueprintType)
+class NIAGARASIMCACHINGEDITOR_API UMovieSceneNiagaraTrackRecorder
+	: public UMovieSceneTrackRecorder
+{
+	GENERATED_BODY()
+
+public:
+	// ~Begin UMovieSceneTrackRecorder Interface
+	virtual void RecordSampleImpl(const FQualifiedFrameTime& CurrentFrameTime) override;
+	virtual void FinalizeTrackImpl() override;
+	virtual void SetSectionStartTimecodeImpl(const FTimecode& InSectionStartTimecode, const FFrameNumber& InSectionFirstFrame) override;
+	virtual UMovieSceneSection* GetMovieSceneSection() const override;
+	virtual void CreateTrackImpl() override;
+	// ~End UMovieSceneTrackRecorder Interface
+
+	/** Returns the Niagara cache track on which the cache manager will be recorded */
+	TWeakObjectPtr<UMovieSceneNiagaraCacheTrack> GetNiagaraCacheTrack() const {return NiagaraCacheTrack;}
+	
+private:
+	
+	/** The NiagaraCache Track to record onto */
+	TWeakObjectPtr<UMovieSceneNiagaraCacheTrack> NiagaraCacheTrack;
+
+	/** Sections to record to on each track*/
+	TWeakObjectPtr<UMovieSceneNiagaraCacheSection> NiagaraCacheSection;
+
+	/** Object to record from */
+	TLazyObjectPtr<UNiagaraComponent> SystemToRecord;
+
+	/** The time at the start of this recording section */
+	double RecordStartTime;
+
+	/** The frame at the start of this recording section */
+	FFrameNumber RecordStartFrame;
+};
+
diff --git a/Engine/Plugins/FX/NiagaraSimCaching/Source/NiagaraSimCachingEditor/Public/Sequencer/NiagaraCacheTrackEditor.h b/Engine/Plugins/FX/NiagaraSimCaching/Source/NiagaraSimCachingEditor/Public/Sequencer/NiagaraCacheTrackEditor.h
new file mode 100644
--- /dev/null
+++ b/Engine/Plugins/FX/NiagaraSimCaching/Source/NiagaraSimCachingEditor/Public/Sequencer/NiagaraCacheTrackEditor.h
@@ -0,0 +1,97 @@
+﻿// Copyright Epic Games, Inc. All Rights Reserved.
+
+#pragma once
+
+#include "ISequencer.h"
+#include "ISequencerSection.h"
+#include "ISequencerTrackEditor.h"
+#include "MovieSceneTrackEditor.h"
+
+/**
+ * Editor for Niagara Cache tracks
+ */
+class FNiagaraCacheTrackEditor : public FMovieSceneTrackEditor
+{
+public:
+	/**
+	 * Constructor
+	 *
+	 * @param InSequencer The sequencer instance to be used by this tool
+	 */
+	FNiagaraCacheTrackEditor( TSharedRef<ISequencer> InSequencer );
+
+	/** Virtual destructor. */
+	virtual ~FNiagaraCacheTrackEditor() override { }
+
+	/**
+	 * Creates an instance of this class.  Called by a sequencer 
+	 *
+	 * @param OwningSequencer The sequencer instance to be used by this tool
+	 * @return The new instance of this class
+	 */
+	static TSharedRef<ISequencerTrackEditor> CreateTrackEditor( TSharedRef<ISequencer> OwningSequencer );
+
+public:
+
+	// ~Begin ISequencerTrackEditor interface
+	virtual void BuildObjectBindingTrackMenu(class FMenuBuilder& MenuBuilder, const TArray<FGuid>& ObjectBindings, const UClass* ObjectClass) override;
+	virtual TSharedRef<ISequencerSection> MakeSectionInterface(UMovieSceneSection& SectionObject, UMovieSceneTrack& Track, FGuid ObjectBinding) override;
+	virtual bool SupportsSequence(UMovieSceneSequence* InSequence) const override;
+	virtual bool SupportsType(TSubclassOf<UMovieSceneTrack> Type) const override;
+	virtual TSharedPtr<SWidget> BuildOutlinerEditWidget(const FGuid& ObjectBinding, UMovieSceneTrack* Track, const FBuildEditWidgetParams& Params) override;
+	virtual const FSlateBrush* GetIconBrush() const override;
+	// ~End ISequencerTrackEditor interface
+	
+private:
+
+	/** Build an internal Niagara cache track */
+	void BuildNiagaraCacheTrack(TArray<FGuid> ObjectBindings, UMovieSceneTrack* Track);
+
+	/** Delegate for AnimatablePropertyChanged in AddKey */
+	FKeyPropertyResult AddKeyInternal(FFrameNumber KeyTime, class UNiagaraComponent* NiagaraComponent, UMovieSceneTrack* Track);
+};
+
+/** Niagara cache sequencer section */
+class FNiagaraCacheSection
+	: public ISequencerSection
+	, public TSharedFromThis<FNiagaraCacheSection>
+{
+public:
+
+	/** Constructor. */
+	FNiagaraCacheSection( UMovieSceneSection& InSection, TWeakPtr<ISequencer> InSequencer);
+
+	/** Virtual destructor. */
+	virtual ~FNiagaraCacheSection() { }
+
+public:
+
+	// ~Begin ISequencerSection interface
+	virtual UMovieSceneSection* GetSectionObject() override;
+	virtual FText GetSectionTitle() const override;
+	virtual float GetSectionHeight() const override;
+	virtual int32 OnPaintSection( FSequencerSectionPainter& Painter ) const override;
+	virtual void BeginResizeSection() override;
+	virtual void ResizeSection(ESequencerSectionResizeMode ResizeMode, FFrameNumber ResizeTime) override;
+	virtual void BeginSlipSection() override;
+	virtual void SlipSection(FFrameNumber SlipTime) override;
+	virtual void BeginDilateSection() override;
+	virtual void DilateSection(const TRange<FFrameNumber>& NewRange, float DilationFactor) override;
+	// ~End ISequencerSection interface
+	
+private:
+
+	void UpdateSection(FFrameNumber& UpdateTime) const;
+
+	/** The section we are visualizing */
+	class UMovieSceneNiagaraCacheSection& Section;
+
+	/** Used to draw animation frame, need selection state and local time*/
+	TWeakPtr<ISequencer> Sequencer;
+
+	/** Cached first loop start offset value valid only during resize */
+	FFrameNumber InitialFirstLoopStartOffsetDuringResize;
+	
+	/** Cached start time valid only during resize */
+	FFrameNumber InitialStartTimeDuringResize;
+};
diff --git a/Engine/Plugins/FX/NiagaraSimCaching/Source/NiagaraSimCachingEditor/Public/Sequencer/TakeRecorderNiagaraCacheSource.h b/Engine/Plugins/FX/NiagaraSimCaching/Source/NiagaraSimCachingEditor/Public/Sequencer/TakeRecorderNiagaraCacheSource.h
new file mode 100644
--- /dev/null
+++ b/Engine/Plugins/FX/NiagaraSimCaching/Source/NiagaraSimCachingEditor/Public/Sequencer/TakeRecorderNiagaraCacheSource.h
@@ -0,0 +1,66 @@
+﻿// Copyright Epic Games, Inc. All Rights Reserved.
+
+#pragma once
+
+#include "CoreMinimal.h"
+#include "NiagaraActor.h"
+#include "NiagaraSimCache.h"
+#include "TakeRecorderSource.h"
+#include "UObject/WeakObjectPtrTemplates.h"
+
+#include "TakeRecorderNiagaraCacheSource.generated.h"
+
+class ULevelSequence;
+class UMovieSceneFolder;
+class UMovieSceneNiagaraCacheTrack;
+class UMovieSceneNiagaraTrackRecorder;
+class ANiagaraCacheManager;
+
+/** A recording source selector for the Niagara integration into take recorder */
+
+UCLASS(Category = "Niagara", meta = (TakeRecorderDisplayName = "Niagara Cache"))
+class UTakeRecorderNiagaraCacheSource : public UTakeRecorderSource
+{
+public:
+	GENERATED_BODY()
+
+	UTakeRecorderNiagaraCacheSource(const FObjectInitializer& ObjInit);
+
+	/** Niagara system to use as take recorder source */
+	UPROPERTY(EditAnywhere, Category = "NiagaraCache")
+	TSoftObjectPtr<ANiagaraActor> SystemToRecord;
+
+	UPROPERTY(EditAnywhere, Category = "NiagaraCache")
+	FNiagaraSimCacheCreateParameters CacheParameters;
+	
+	// ~Begin UObject Interface
+	virtual void PostEditChangeProperty(FPropertyChangedEvent& PropertyChangedEvent) override;
+	virtual void PostLoad() override;
+	// ~End UObject Interface
+	
+protected:
+	// ~Begin UTakeRecorderSource Interface
+	virtual bool SupportsSubscenes() const override { return false; }
+	// ~End UTakeRecorderSource Interface
+
+private:
+	// ~Begin UTakeRecorderSource Interface
+	virtual TArray<UTakeRecorderSource*> PreRecording(ULevelSequence* InSequence, FMovieSceneSequenceID InSequenceID, ULevelSequence* InMasterSequence, FManifestSerializer* InManifestSerializer) override;
+	virtual void StartRecording(const FTimecode& InSectionStartTimecode, const FFrameNumber& InSectionFirstFrame, class ULevelSequence* InSequence) override;
+	virtual void StopRecording(class ULevelSequence* InSequence) override;
+	virtual void TickRecording(const FQualifiedFrameTime& CurrentTime) override;
+	virtual TArray<UTakeRecorderSource*> PostRecording(class ULevelSequence* InSequence, ULevelSequence* InMasterSequence, const bool bCancelled) override;
+	virtual void AddContentsToFolder(UMovieSceneFolder* InFolder) override;
+	virtual FText GetDisplayTextImpl() const override;
+	// ~End UTakeRecorderSource
+	
+	/** Niagara cache track recorder used by this source */
+	UPROPERTY()
+	TObjectPtr<UMovieSceneNiagaraTrackRecorder> TrackRecorder;
+
+	/**
+	 * Stores an existing Niagara cache track track in the Sequence to be recorded or
+	 * a new one created for recording. Set during PreRecording.
+	 */
+	TWeakObjectPtr<UMovieSceneNiagaraCacheTrack> CachedNiagaraCacheTrack;
+};
