Fixed variable naming

diff --git a/Engine/Plugins/FX/NiagaraSimCaching/Source/NiagaraSimCachingEditor/Private/Sequencer/TakeRecorderNiagaraCacheSource.cpp b/Engine/Plugins/FX/NiagaraSimCaching/Source/NiagaraSimCachingEditor/Private/Sequencer/TakeRecorderNiagaraCacheSource.cpp
--- a/Engine/Plugins/FX/NiagaraSimCaching/Source/NiagaraSimCachingEditor/Private/Sequencer/TakeRecorderNiagaraCacheSource.cpp
+++ b/Engine/Plugins/FX/NiagaraSimCaching/Source/NiagaraSimCachingEditor/Private/Sequencer/TakeRecorderNiagaraCacheSource.cpp
@@ -21,7 +21,7 @@ UTakeRecorderNiagaraCacheSource::UTakeRecorderNiagaraCacheSource(const FObjectIn
 	TrackTint = FColor(0.0f, 125.0f, 255.0f, 65.0f);
 }
 
-TArray<UTakeRecorderSource*> UTakeRecorderNiagaraCacheSource::PreRecording(ULevelSequence* InSequence, FMovieSceneSequenceID InSequenceID, ULevelSequence* InMasterSequence, FManifestSerializer* InManifestSerializer)
+TArray<UTakeRecorderSource*> UTakeRecorderNiagaraCacheSource::PreRecording(ULevelSequence* InSequence, FMovieSceneSequenceID InSequenceID, ULevelSequence* InParentSequence, FManifestSerializer* InManifestSerializer)
 {
 	if (SystemToRecord)
 	{
@@ -34,9 +34,9 @@ TArray<UTakeRecorderSource*> UTakeRecorderNiagaraCacheSource::PreRecording(ULeve
 		// Would be better to use a TSoftObjectPtr for the actor and have a mechanism similar to the FixUpPIE to go
 		// from PIE to editor to have a correct soft object path
 		//const FGuid ObjectGuid = MovieScene->AddPossessable(NiagaraCacheManager->GetActorLabel(), NiagaraCacheManager->GetClass());
-		//InMasterSequence->BindPossessableObject(ObjectGuid, *NiagaraCacheManager, NiagaraCacheManager->GetWorld());
+		//InParentSequence->BindPossessableObject(ObjectGuid, *NiagaraCacheManager, NiagaraCacheManager->GetWorld());
 
-		//InMasterSequence->
+		//InParentSequence->
 		//TrackRecorder->CreateTrack(MovieScene, SystemToRecord->GetNiagaraComponent(), MovieScene, nullptr, ObjectGuid);
 		CachedNiagaraCacheTrack = TrackRecorder->GetNiagaraCacheTrack();
 	}
@@ -67,7 +67,7 @@ void UTakeRecorderNiagaraCacheSource::StopRecording(class ULevelSequence* InSequ
 	}
 }
 
-TArray<UTakeRecorderSource*> UTakeRecorderNiagaraCacheSource::PostRecording(class ULevelSequence* InSequence, ULevelSequence* InMasterSequence, const bool bCancelled)
+TArray<UTakeRecorderSource*> UTakeRecorderNiagaraCacheSource::PostRecording(class ULevelSequence* InSequence, ULevelSequence* InParentSequence, const bool bCancelled)
 {
 	if (TrackRecorder)
 	{
diff --git a/Engine/Plugins/FX/NiagaraSimCaching/Source/NiagaraSimCachingEditor/Public/Sequencer/TakeRecorderNiagaraCacheSource.h b/Engine/Plugins/FX/NiagaraSimCaching/Source/NiagaraSimCachingEditor/Public/Sequencer/TakeRecorderNiagaraCacheSource.h
--- a/Engine/Plugins/FX/NiagaraSimCaching/Source/NiagaraSimCachingEditor/Public/Sequencer/TakeRecorderNiagaraCacheSource.h
+++ b/Engine/Plugins/FX/NiagaraSimCaching/Source/NiagaraSimCachingEditor/Public/Sequencer/TakeRecorderNiagaraCacheSource.h
@@ -45,11 +45,11 @@ protected:
 
 private:
 	// ~Begin UTakeRecorderSource Interface
-	virtual TArray<UTakeRecorderSource*> PreRecording(ULevelSequence* InSequence, FMovieSceneSequenceID InSequenceID, ULevelSequence* InMasterSequence, FManifestSerializer* InManifestSerializer) override;
+	virtual TArray<UTakeRecorderSource*> PreRecording(ULevelSequence* InSequence, FMovieSceneSequenceID InSequenceID, ULevelSequence* InParentSequence, FManifestSerializer* InManifestSerializer) override;
 	virtual void StartRecording(const FTimecode& InSectionStartTimecode, const FFrameNumber& InSectionFirstFrame, class ULevelSequence* InSequence) override;
 	virtual void StopRecording(class ULevelSequence* InSequence) override;
 	virtual void TickRecording(const FQualifiedFrameTime& CurrentTime) override;
-	virtual TArray<UTakeRecorderSource*> PostRecording(class ULevelSequence* InSequence, ULevelSequence* InMasterSequence, const bool bCancelled) override;
+	virtual TArray<UTakeRecorderSource*> PostRecording(class ULevelSequence* InSequence, ULevelSequence* InParentSequence, const bool bCancelled) override;
 	virtual void AddContentsToFolder(UMovieSceneFolder* InFolder) override;
 	virtual FText GetDisplayTextImpl() const override;
 	// ~End UTakeRecorderSource
