Added missing plugin descriptor file

diff --git a/Engine/Plugins/FX/NiagaraSimCaching/NiagaraSimCaching.uplugin b/Engine/Plugins/FX/NiagaraSimCaching/NiagaraSimCaching.uplugin
new file mode 100644
--- /dev/null
+++ b/Engine/Plugins/FX/NiagaraSimCaching/NiagaraSimCaching.uplugin
@@ -0,0 +1,40 @@
+{
+	"FileVersion": 3,
+	"Version": 1,
+	"VersionName": "1.0",
+	"FriendlyName": "NiagaraSimCaching",
+	"Description": "Adds support for recording and playing back Niagara simulations in sequencer via take recorder",
+	"Category": "FX",
+	"CreatedBy": "Epic Games, Inc.",
+	"CreatedByURL": "http://epicgames.com",
+	"DocsURL": "",
+	"MarketplaceURL": "",
+	"SupportURL": "",
+	"EnabledByDefault": false,
+	"CanContainContent": false,
+	"IsBetaVersion": true,
+	"IsExperimentalVersion": true,
+	"Installed": false,
+	"Modules": [
+		{
+			"Name": "NiagaraSimCaching",
+			"Type": "Runtime",
+			"LoadingPhase": "PreDefault"
+		},
+		{
+			"Name": "NiagaraSimCachingEditor",
+			"Type": "Editor",
+			"LoadingPhase": "PreDefault"
+		}
+	],
+	"Plugins": [
+		{
+			"Name": "Takes",
+			"Enabled": true
+		},
+		{
+			"Name": "Niagara",
+			"Enabled": true
+		}
+	]
+}
\ No newline at end of file
