Added icon for cache track record action

diff --git a/Engine/Plugins/FX/NiagaraSimCaching/NiagaraSimCaching.uplugin b/Engine/Plugins/FX/NiagaraSimCaching/NiagaraSimCaching.uplugin
--- a/Engine/Plugins/FX/NiagaraSimCaching/NiagaraSimCaching.uplugin
+++ b/Engine/Plugins/FX/NiagaraSimCaching/NiagaraSimCaching.uplugin
@@ -11,7 +11,7 @@
 	"MarketplaceURL": "",
 	"SupportURL": "",
 	"EnabledByDefault": false,
-	"CanContainContent": false,
+	"CanContainContent": true,
 	"IsBetaVersion": true,
 	"IsExperimentalVersion": true,
 	"Installed": false,
