Small code review changes

diff --git a/Engine/Plugins/FX/Niagara/Source/Niagara/Private/MovieScene/MovieSceneNiagaraSystemTrackTemplate.cpp b/Engine/Plugins/FX/Niagara/Source/Niagara/Private/MovieScene/MovieSceneNiagaraSystemTrackTemplate.cpp
--- a/Engine/Plugins/FX/Niagara/Source/Niagara/Private/MovieScene/MovieSceneNiagaraSystemTrackTemplate.cpp
+++ b/Engine/Plugins/FX/Niagara/Source/Niagara/Private/MovieScene/MovieSceneNiagaraSystemTrackTemplate.cpp
@@ -282,10 +282,6 @@ FMovieSceneNiagaraSystemTrackImplementation::FMovieSceneNiagaraSystemTrackImplem
 
 void FMovieSceneNiagaraSystemTrackImplementation::Evaluate(const FMovieSceneEvaluationTrack& Track, TArrayView<const FMovieSceneFieldEntry_ChildTemplate> Children, const FMovieSceneEvaluationOperand& Operand, const FMovieSceneContext& Context, const FPersistentEvaluationData& PersistentData, FMovieSceneExecutionTokens& ExecutionTokens) const
 {
-	if (Track.GetSourceTrack()->IsEvalDisabled())
-	{
-		return;
-	}
 	ExecutionTokens.SetContext(Context);
 
 	// only add a token if there isn't one already, otherwise another track's token takes precendence
diff --git a/Engine/Plugins/FX/NiagaraSimCaching/Source/NiagaraSimCaching/Private/Niagara/Sequencer/MovieSceneNiagaraCacheTemplate.cpp b/Engine/Plugins/FX/NiagaraSimCaching/Source/NiagaraSimCaching/Private/Niagara/Sequencer/MovieSceneNiagaraCacheTemplate.cpp
--- a/Engine/Plugins/FX/NiagaraSimCaching/Source/NiagaraSimCaching/Private/Niagara/Sequencer/MovieSceneNiagaraCacheTemplate.cpp
+++ b/Engine/Plugins/FX/NiagaraSimCaching/Source/NiagaraSimCaching/Private/Niagara/Sequencer/MovieSceneNiagaraCacheTemplate.cpp
@@ -48,8 +48,8 @@ struct FPreAnimatedNiagaraCacheTokenProducer : IMovieScenePreAnimatedTokenProduc
 /** A movie scene execution token that executes a NiagaraCache */
 struct FNiagaraCacheExecutionToken : IMovieSceneExecutionToken
 {
-	FNiagaraCacheExecutionToken(const TArray<FMovieSceneNiagaraSectionTemplateParameter> &InSections) :
-		CacheSections(InSections)
+	FNiagaraCacheExecutionToken(const TOptional<FMovieSceneNiagaraSectionTemplateParameter>& InSection, const ENiagaraSimCacheSectionPlayMode& TargetReplayMode) :
+		CacheSection(InSection), ReplayMode(TargetReplayMode)
 	{}
 
 	/** Execute this token, operating on all objects referenced by 'Operand' */
@@ -73,41 +73,37 @@ struct FNiagaraCacheExecutionToken : IMovieSceneExecutionToken
 		{
 			return;
 		}
-		
-		ENiagaraSimCacheSectionPlayMode ReplayMode = ENiagaraSimCacheSectionPlayMode::DisplayCacheOnly;
-		for (const FMovieSceneNiagaraSectionTemplateParameter& Section : CacheSections)
+
+		if (!CacheSection.IsSet())
 		{
-			ReplayMode = Section.Params.CacheReplayPlayMode;
-			if (Section.IsTimeWithinSection(Context.GetTime().GetFrame()) && Section.Params.SimCache)
+			// outside of cached sections, use configured replay mode
+			if (ReplayMode == ENiagaraSimCacheSectionPlayMode::DisplayCacheOnly)
 			{
-				if (NiagaraComponent->GetSimCache() != Section.Params.SimCache)
-				{
-					NiagaraComponent->SetSimCache(Section.Params.SimCache);
-				}
-		
-				// calculate the time at which to evaluate the sim cache
-				float Age = MapTimeToAnimation(Section.Params, Section.Params.SimCache->GetDurationSeconds(), Context.GetTime(), Context.GetFrameRate(), Section) + Section.Params.SimCache->GetStartSeconds();
-
-				FScopedPreAnimatedCaptureSource CaptureSource(&Player.PreAnimatedState, PersistentData.GetSectionKey(), true);
-				Player.SavePreAnimatedState(*NiagaraComponent, FPreAnimatedNiagaraCacheTokenProducer::GetAnimTypeID(), FPreAnimatedNiagaraCacheTokenProducer());
-				NiagaraComponent->SetDesiredAge(Age);
-				NiagaraComponent->SetAgeUpdateMode(ENiagaraAgeUpdateMode::DesiredAge);
-				if (!NiagaraComponent->IsActive())
-				{
-					NiagaraComponent->ResetSystem();
-				}
-				return;
+				NiagaraComponent->DeactivateImmediate();
+			}
+			else if (ReplayMode == ENiagaraSimCacheSectionPlayMode::SimWithoutCache)
+			{
+				NiagaraComponent->SetSimCache(nullptr);
 			}
+			return;
 		}
-
-		// outside of cached sections, use configured replay mode
-		if (ReplayMode == ENiagaraSimCacheSectionPlayMode::DisplayCacheOnly)
+		
+		FMovieSceneNiagaraSectionTemplateParameter& Section = CacheSection.GetValue();
+		if (NiagaraComponent->GetSimCache() != Section.Params.SimCache)
 		{
-			NiagaraComponent->DeactivateImmediate();
+			NiagaraComponent->SetSimCache(Section.Params.SimCache);
 		}
-		else if (ReplayMode == ENiagaraSimCacheSectionPlayMode::SimWithoutCache)
+
+		// calculate the time at which to evaluate the sim cache
+		float Age = MapTimeToAnimation(Section.Params, Section.Params.SimCache->GetDurationSeconds(), Context.GetTime(), Context.GetFrameRate(), Section) + Section.Params.SimCache->GetStartSeconds();
+
+		FScopedPreAnimatedCaptureSource CaptureSource(&Player.PreAnimatedState, PersistentData.GetSectionKey(), true);
+		Player.SavePreAnimatedState(*NiagaraComponent, FPreAnimatedNiagaraCacheTokenProducer::GetAnimTypeID(), FPreAnimatedNiagaraCacheTokenProducer());
+		NiagaraComponent->SetDesiredAge(Age);
+		NiagaraComponent->SetAgeUpdateMode(ENiagaraAgeUpdateMode::DesiredAge);
+		if (!NiagaraComponent->IsActive())
 		{
-			NiagaraComponent->SetSimCache(nullptr);
+			NiagaraComponent->ResetSystem();
 		}
 	}
 
@@ -148,7 +144,8 @@ struct FNiagaraCacheExecutionToken : IMovieSceneExecutionToken
 		return AnimPosition;
 	}
 
-	TArray<FMovieSceneNiagaraSectionTemplateParameter> CacheSections;
+	TOptional<FMovieSceneNiagaraSectionTemplateParameter> CacheSection;
+	ENiagaraSimCacheSectionPlayMode ReplayMode;
 };
 
 struct FNiagaraCacheResetToken : IMovieSceneExecutionToken
@@ -192,8 +189,22 @@ void FMovieSceneNiagaraCacheSectionTemplate::Evaluate(const FMovieSceneEvaluatio
 	}
 	ExecutionTokens.SetContext(Context);
 
-	// remove the sim cache if we don't have section data to display
-	if (NiagaraTrack->IsEvalDisabled() || CacheSections.Num() == 0)
+	ENiagaraSimCacheSectionPlayMode TargetReplayMode = ENiagaraSimCacheSectionPlayMode::DisplayCacheOnly;
+	TOptional<FMovieSceneNiagaraSectionTemplateParameter> CacheSectionToApply;
+	bool bHasValidSections = false;
+	for (const FMovieSceneNiagaraSectionTemplateParameter& Section : CacheSections)
+	{
+		bHasValidSections |= Section.Params.SimCache && Section.Params.SimCache->IsCacheValid();
+		TargetReplayMode = Section.Params.CacheReplayPlayMode;
+		if (Section.IsTimeWithinSection(Context.GetTime().GetFrame()) && Section.Params.SimCache)
+		{
+			CacheSectionToApply = Section;
+			break;
+		}
+	}
+
+	// remove the sim cache if we don't have section data to display anywhere
+	if (!bHasValidSections)
 	{
 		ExecutionTokens.Add(FNiagaraCacheResetToken());
 		return;
@@ -204,14 +215,14 @@ void FMovieSceneNiagaraCacheSectionTemplate::Evaluate(const FMovieSceneEvaluatio
 	FNiagaraSharedMarkerToken* SharedCacheToken = static_cast<FNiagaraSharedMarkerToken*>(ExecutionTokens.FindShared(TokenID));
 	if (SharedCacheToken == nullptr)
 	{
-		ExecutionTokens.Add(FNiagaraCacheExecutionToken(CacheSections));
+		ExecutionTokens.Add(FNiagaraCacheExecutionToken(CacheSectionToApply, TargetReplayMode));
 		FNiagaraSharedMarkerToken SharedToken = FNiagaraSharedMarkerToken();
 		SharedToken.BoundObjectIDs.Add(Operand.ObjectBindingID);
 		ExecutionTokens.AddShared(TokenID, SharedToken);
 	}
 	else if (!SharedCacheToken->BoundObjectIDs.Contains(Operand.ObjectBindingID))
 	{
-		ExecutionTokens.Add(FNiagaraCacheExecutionToken(CacheSections));
+		ExecutionTokens.Add(FNiagaraCacheExecutionToken(CacheSectionToApply, TargetReplayMode));
 		SharedCacheToken->BoundObjectIDs.Add(Operand.ObjectBindingID);
 	}
 }
