Fixed non unity/pch errors found when building UnrealGame and UnrealEditor without unity/pch

diff --git a/Engine/Plugins/Experimental/FullBodyIK/Source/PBIK/Public/Core/PBIKConstraint.h b/Engine/Plugins/Experimental/FullBodyIK/Source/PBIK/Public/Core/PBIKConstraint.h
--- a/Engine/Plugins/Experimental/FullBodyIK/Source/PBIK/Public/Core/PBIKConstraint.h
+++ b/Engine/Plugins/Experimental/FullBodyIK/Source/PBIK/Public/Core/PBIKConstraint.h
@@ -5,6 +5,7 @@
 #include "Math/Quat.h"
 
 namespace PBIK { struct FRigidBody; }
+struct FPBIKSolverSettings;
 
 namespace PBIK
 {
diff --git a/Engine/Plugins/Experimental/VirtualCamera/Source/PixelStreamingVCam/Private/VCamPixelStreamingSubsystem.cpp b/Engine/Plugins/Experimental/VirtualCamera/Source/PixelStreamingVCam/Private/VCamPixelStreamingSubsystem.cpp
--- a/Engine/Plugins/Experimental/VirtualCamera/Source/PixelStreamingVCam/Private/VCamPixelStreamingSubsystem.cpp
+++ b/Engine/Plugins/Experimental/VirtualCamera/Source/PixelStreamingVCam/Private/VCamPixelStreamingSubsystem.cpp
@@ -1,6 +1,7 @@
 // Copyright Epic Games, Inc. All Rights Reserved.
 
 #include "VCamPixelStreamingSubsystem.h"
+#include "Features/IModularFeatures.h"
 #include "ILiveLinkClient.h"
 #include "VCamPixelStreamingLiveLink.h"
 #include "VCamPixelStreamingSession.h"
diff --git a/Engine/Plugins/FX/NiagaraSimCaching/Source/NiagaraSimCachingEditor/Private/Sequencer/MovieSceneNiagaraTrackRecorder.cpp b/Engine/Plugins/FX/NiagaraSimCaching/Source/NiagaraSimCachingEditor/Private/Sequencer/MovieSceneNiagaraTrackRecorder.cpp
--- a/Engine/Plugins/FX/NiagaraSimCaching/Source/NiagaraSimCachingEditor/Private/Sequencer/MovieSceneNiagaraTrackRecorder.cpp
+++ b/Engine/Plugins/FX/NiagaraSimCaching/Source/NiagaraSimCachingEditor/Private/Sequencer/MovieSceneNiagaraTrackRecorder.cpp
@@ -27,7 +27,7 @@ UMovieSceneTrackRecorder* FMovieSceneNiagaraTrackRecorderFactory::CreateTrackRec
 	return NewObject<UMovieSceneNiagaraTrackRecorder>();
 }
 
-UMovieSceneTrackRecorder* FMovieSceneNiagaraTrackRecorderFactory::CreateTrackRecorderForCacheTrack(IMovieSceneCachedTrack* CachedTrack, TObjectPtr<ULevelSequence> Sequence, TSharedPtr<ISequencer> Sequencer) const
+UMovieSceneTrackRecorder* FMovieSceneNiagaraTrackRecorderFactory::CreateTrackRecorderForCacheTrack(IMovieSceneCachedTrack* CachedTrack, const TObjectPtr<ULevelSequence>& Sequence, const TSharedPtr<ISequencer>& Sequencer) const
 {
 	if (UMovieSceneNiagaraCacheTrack* NiagaraCacheTrack = Cast<UMovieSceneNiagaraCacheTrack>(CachedTrack))
 	{
diff --git a/Engine/Plugins/FX/NiagaraSimCaching/Source/NiagaraSimCachingEditor/Public/Sequencer/MovieSceneNiagaraTrackRecorder.h b/Engine/Plugins/FX/NiagaraSimCaching/Source/NiagaraSimCachingEditor/Public/Sequencer/MovieSceneNiagaraTrackRecorder.h
--- a/Engine/Plugins/FX/NiagaraSimCaching/Source/NiagaraSimCachingEditor/Public/Sequencer/MovieSceneNiagaraTrackRecorder.h
+++ b/Engine/Plugins/FX/NiagaraSimCaching/Source/NiagaraSimCachingEditor/Public/Sequencer/MovieSceneNiagaraTrackRecorder.h
@@ -24,7 +24,7 @@ public:
 	// ~Begin IMovieSceneTrackRecorderFactory Interface
 	virtual bool CanRecordObject(class UObject* InObjectToRecord) const override;
 	virtual UMovieSceneTrackRecorder* CreateTrackRecorderForObject() const override;
-	virtual class UMovieSceneTrackRecorder* CreateTrackRecorderForCacheTrack(IMovieSceneCachedTrack* CachedTrack, TObjectPtr<ULevelSequence> Sequence, TSharedPtr<ISequencer> Sequencer) const override;
+	virtual class UMovieSceneTrackRecorder* CreateTrackRecorderForCacheTrack(IMovieSceneCachedTrack* CachedTrack, const TObjectPtr<ULevelSequence>& Sequence, const TSharedPtr<ISequencer>& Sequencer) const override;
 
 	// Particle Systems are entire components and you can't animate them as a property
 	virtual bool CanRecordProperty(class UObject* InObjectToRecord, class FProperty* InPropertyToRecord) const override { return false; }
diff --git a/Engine/Plugins/MovieScene/SequencerScripting/Source/SequencerScripting/Public/ExtensionLibraries/MovieSceneSequenceExtensions.h b/Engine/Plugins/MovieScene/SequencerScripting/Source/SequencerScripting/Public/ExtensionLibraries/MovieSceneSequenceExtensions.h
--- a/Engine/Plugins/MovieScene/SequencerScripting/Source/SequencerScripting/Public/ExtensionLibraries/MovieSceneSequenceExtensions.h
+++ b/Engine/Plugins/MovieScene/SequencerScripting/Source/SequencerScripting/Public/ExtensionLibraries/MovieSceneSequenceExtensions.h
@@ -4,7 +4,7 @@
 
 #include "Kismet/BlueprintFunctionLibrary.h"
 #include "MovieSceneObjectBindingID.h" // for EMovieSceneObjectBindingSpace
-
+#include "MovieSceneTrack.h"
 #include "MovieSceneSequenceExtensions.generated.h"
 
 class UMovieScene;
diff --git a/Engine/Plugins/Online/OnlineServices/Source/OnlineServicesCommon/Public/Online/OnlineExecHandler.h b/Engine/Plugins/Online/OnlineServices/Source/OnlineServicesCommon/Public/Online/OnlineExecHandler.h
--- a/Engine/Plugins/Online/OnlineServices/Source/OnlineServicesCommon/Public/Online/OnlineExecHandler.h
+++ b/Engine/Plugins/Online/OnlineServices/Source/OnlineServicesCommon/Public/Online/OnlineExecHandler.h
@@ -2,6 +2,7 @@
 
 #pragma once
 
+#include "HAL/IConsoleManager.h"
 #include "Online/OnlineAsyncOpHandle.h"
 #include "Online/OnlineServicesLog.h"
 #include "Online/OnlineResult.h"
diff --git a/Engine/Plugins/Runtime/GameplayBehaviorSmartObjects/Source/GameplayBehaviorSmartObjectsModule/Public/AI/BTTask_FindAndUseGameplayBehaviorSmartObject.h b/Engine/Plugins/Runtime/GameplayBehaviorSmartObjects/Source/GameplayBehaviorSmartObjectsModule/Public/AI/BTTask_FindAndUseGameplayBehaviorSmartObject.h
--- a/Engine/Plugins/Runtime/GameplayBehaviorSmartObjects/Source/GameplayBehaviorSmartObjectsModule/Public/AI/BTTask_FindAndUseGameplayBehaviorSmartObject.h
+++ b/Engine/Plugins/Runtime/GameplayBehaviorSmartObjects/Source/GameplayBehaviorSmartObjectsModule/Public/AI/BTTask_FindAndUseGameplayBehaviorSmartObject.h
@@ -7,9 +7,8 @@
 #include "BTTask_FindAndUseGameplayBehaviorSmartObject.generated.h"
 
 class UAITask_UseGameplayBehaviorSmartObject;
-
-
 class AITask_UseSmartObject;
+struct FSmartObjectClaimHandle;
 
 struct FBTUseSOTaskMemory
 {
diff --git a/Engine/Plugins/Runtime/GeometryCache/Source/GeometryCache/Private/GeometryCacheStreamingManager.cpp b/Engine/Plugins/Runtime/GeometryCache/Source/GeometryCache/Private/GeometryCacheStreamingManager.cpp
--- a/Engine/Plugins/Runtime/GeometryCache/Source/GeometryCache/Private/GeometryCacheStreamingManager.cpp
+++ b/Engine/Plugins/Runtime/GeometryCache/Source/GeometryCache/Private/GeometryCacheStreamingManager.cpp
@@ -6,6 +6,7 @@
 #include "GeometryCacheComponent.h"
 #include "GeometryCacheModule.h"
 #include "GeometryCacheTrackStreamable.h"
+#include "RenderingThread.h"
 #include "StreamingGeometryCacheData.h"
 #include "Async/Async.h"
 #include "Misc/ScopeRWLock.h"
diff --git a/Engine/Plugins/VirtualProduction/Takes/Source/TakeTrackRecorders/Public/TrackRecorders/IMovieSceneTrackRecorderFactory.h b/Engine/Plugins/VirtualProduction/Takes/Source/TakeTrackRecorders/Public/TrackRecorders/IMovieSceneTrackRecorderFactory.h
--- a/Engine/Plugins/VirtualProduction/Takes/Source/TakeTrackRecorders/Public/TrackRecorders/IMovieSceneTrackRecorderFactory.h
+++ b/Engine/Plugins/VirtualProduction/Takes/Source/TakeTrackRecorders/Public/TrackRecorders/IMovieSceneTrackRecorderFactory.h
@@ -59,7 +59,7 @@ public:
 	*/
 	virtual class UMovieSceneTrackRecorder* CreateTrackRecorderForProperty(UObject* InObjectToRecord, const FName& InPropertyToRecord) const = 0;
 
-	virtual class UMovieSceneTrackRecorder* CreateTrackRecorderForCacheTrack(IMovieSceneCachedTrack* CachedTrack, TObjectPtr<ULevelSequence> Sequence, TSharedPtr<ISequencer> Sequencer) const { return nullptr; };
+	virtual class UMovieSceneTrackRecorder* CreateTrackRecorderForCacheTrack(IMovieSceneCachedTrack* CachedTrack, const TObjectPtr<ULevelSequence>& Sequence, const TSharedPtr<ISequencer>& Sequencer) const { return nullptr; };
 
 	/**
 	* Get the human readable display name for the recorder. Used for debugging purposes to help identify which factory
