Fixed a bug where the length of the life cycle track was not evaluated correctly when caching a Niagara simulation

diff --git a/Engine/Plugins/FX/NiagaraSimCaching/Source/NiagaraSimCachingEditor/Private/Sequencer/MovieSceneNiagaraTrackRecorder.cpp b/Engine/Plugins/FX/NiagaraSimCaching/Source/NiagaraSimCachingEditor/Private/Sequencer/MovieSceneNiagaraTrackRecorder.cpp
--- a/Engine/Plugins/FX/NiagaraSimCaching/Source/NiagaraSimCachingEditor/Private/Sequencer/MovieSceneNiagaraTrackRecorder.cpp
+++ b/Engine/Plugins/FX/NiagaraSimCaching/Source/NiagaraSimCachingEditor/Private/Sequencer/MovieSceneNiagaraTrackRecorder.cpp
@@ -233,7 +233,7 @@ void UMovieSceneNiagaraTrackRecorder::RecordSampleImpl(const FQualifiedFrameTime
 			
 			// Expand the section to the new length
 			FFrameNumber EndFrame = CurrentFrameTime.Time.CeilToFrame();
-			NiagaraCacheSection->ExpandToFrame(EndFrame);
+			NiagaraCacheSection->SetEndFrame(EndFrame);
 		}
 		
 		for (const FString& Warning : FeedbackContext.Warnings)
