Add a warning if life cycle track is not set up correctly for sim caching

diff --git a/Engine/Plugins/FX/Niagara/Source/Niagara/Public/MovieScene/MovieSceneNiagaraSystemSpawnSection.h b/Engine/Plugins/FX/Niagara/Source/Niagara/Public/MovieScene/MovieSceneNiagaraSystemSpawnSection.h
--- a/Engine/Plugins/FX/Niagara/Source/Niagara/Public/MovieScene/MovieSceneNiagaraSystemSpawnSection.h
+++ b/Engine/Plugins/FX/Niagara/Source/Niagara/Public/MovieScene/MovieSceneNiagaraSystemSpawnSection.h
@@ -51,7 +51,7 @@ public:
 
 	ENiagaraSystemSpawnSectionEndBehavior GetSectionEndBehavior() const;
 
-	ENiagaraAgeUpdateMode GetAgeUpdateMode() const;
+	NIAGARA_API ENiagaraAgeUpdateMode GetAgeUpdateMode() const;
 
 	bool GetAllowScalability()const;
 
diff --git a/Engine/Plugins/FX/NiagaraSimCaching/Source/NiagaraSimCachingEditor/Private/Sequencer/NiagaraCacheTrackEditor.cpp b/Engine/Plugins/FX/NiagaraSimCaching/Source/NiagaraSimCachingEditor/Private/Sequencer/NiagaraCacheTrackEditor.cpp
--- a/Engine/Plugins/FX/NiagaraSimCaching/Source/NiagaraSimCachingEditor/Private/Sequencer/NiagaraCacheTrackEditor.cpp
+++ b/Engine/Plugins/FX/NiagaraSimCaching/Source/NiagaraSimCachingEditor/Private/Sequencer/NiagaraCacheTrackEditor.cpp
@@ -12,6 +12,8 @@
 #include "Niagara/NiagaraSimCachingEditorStyle.h"
 #include "Niagara/Sequencer/MovieSceneNiagaraCacheSection.h"
 #include "Niagara/Sequencer/MovieSceneNiagaraCacheTrack.h"
+#include "MovieScene/MovieSceneNiagaraSystemSpawnSection.h"
+#include "MovieScene/MovieSceneNiagaraSystemTrack.h"
 #include "NiagaraComponent.h"
 #include "NiagaraSimCache.h"
 #include "NiagaraSystem.h"
@@ -414,6 +416,73 @@ FReply FNiagaraCacheTrackEditor::RecordCacheTrack(IMovieSceneCachedTrack* Track)
 	return FReply::Handled();
 }
 
+bool FNiagaraCacheTrackEditor::IsCacheTrackOutOfDate(UMovieSceneTrack* Track)
+{
+	if (UMovieSceneNiagaraCacheTrack* NiagaraTrack = Cast<UMovieSceneNiagaraCacheTrack>(Track))
+	{
+		for (UMovieSceneSection* Section : NiagaraTrack->GetAllSections())
+		{
+			if (UMovieSceneNiagaraCacheSection* NiagaraCacheSection = Cast<UMovieSceneNiagaraCacheSection>(Section))
+			{
+				if (NiagaraCacheSection->bCacheOutOfDate)
+				{
+					return true;
+				}
+			}
+		}
+	}
+	return false;
+}
+
+bool FNiagaraCacheTrackEditor::HasConflictingLifecycleTrack(UMovieSceneTrack* CacheTrack) const
+{
+	UMovieScene* MovieScene = GetMovieSceneSequence()->GetMovieScene();
+	FFrameNumber PlaybackStartFrame = MovieScene->GetPlaybackRange().GetLowerBoundValue();
+	const TArray<FMovieSceneBinding>& SceneBindings = MovieScene->GetBindings();
+
+	for (const FMovieSceneBinding& Binding : SceneBindings)
+	{
+		TArray<UMovieSceneTrack*> ComponentTracks = Binding.GetTracks();
+		// find any life cycle tracks bound to the same component
+		if (ComponentTracks.Contains(CacheTrack))
+		{
+			for (UMovieSceneTrack* Track : ComponentTracks)
+			{
+				if (UMovieSceneNiagaraSystemTrack* SystemTrack = Cast<UMovieSceneNiagaraSystemTrack>(Track))
+				{
+					if (SystemTrack->IsEvalDisabled() == false)
+					{
+						for (UMovieSceneSection* Section : SystemTrack->GetAllSections())
+						{
+							// check if we start at the same frame as playback. If we do and don't use desired age then that's a problem, as recording a cache will reset the system multiple times at the start and invalidate the recording.
+							UMovieSceneNiagaraSystemSpawnSection* SpawnSection = Cast<UMovieSceneNiagaraSystemSpawnSection>(Section);
+							if (SpawnSection && SpawnSection->GetAgeUpdateMode() == ENiagaraAgeUpdateMode::TickDeltaTime && SpawnSection->GetInclusiveStartFrame() == PlaybackStartFrame)
+							{
+								return true;
+							}
+						}
+					}
+				}
+			}
+		}
+	}
+	return false;
+}
+
+FText FNiagaraCacheTrackEditor::GetCacheTrackWarnings(UMovieSceneTrack* Track) const
+{
+	TArray<FText> Warnings;
+	if (HasConflictingLifecycleTrack(Track))
+	{
+		Warnings.Add(LOCTEXT("AddNiagaraCache_LifecycleTrackWarn", "The system life cycle track associated with this cache track is not set up well for cache recording.\nEither (1) set it to use desired age mode or (2) move it so it starts on a different frame than the sequence playback (e.g. one frame earlier)."));
+	}
+	if (IsCacheTrackOutOfDate(Track))
+	{
+		Warnings.Add(LOCTEXT("AddNiagaraCache_OutOfDate", "This cache track is out of date, as it's properties were changed after recording. Consider re-recording the cached data."));
+	}
+	return FText::Join(FText::FromString("\n\n"), Warnings);
+}
+
 TSharedPtr<SWidget> FNiagaraCacheTrackEditor::BuildOutlinerEditWidget(const FGuid& ObjectBinding, UMovieSceneTrack* Track, const FBuildEditWidgetParams& Params)
 {
 	TSharedPtr<ISequencer> SequencerPtr = GetSequencer();
@@ -447,22 +516,13 @@ TSharedPtr<SWidget> FNiagaraCacheTrackEditor::BuildOutlinerEditWidget(const FGui
 		.VAlign(VAlign_Center)
 		[
 			SNew(SImage)
-			.ToolTipText(LOCTEXT("AddNiagaraCache_OutOfDate", "This cache track is out of date, as it's properties were changed after recording. Consider re-recording the cached data."))
+			.ToolTipText_Raw(this, &FNiagaraCacheTrackEditor::GetCacheTrackWarnings, Track)
 			.Image(FAppStyle::Get().GetBrush("Icons.WarningWithColor"))
-			.Visibility_Lambda([Track]()
+			.Visibility_Lambda([Track, this]()
 			{
-				if (UMovieSceneNiagaraCacheTrack* NiagaraTrack = Cast<UMovieSceneNiagaraCacheTrack>(Track))
+				if (IsCacheTrackOutOfDate(Track) || HasConflictingLifecycleTrack(Track))
 				{
-					for (UMovieSceneSection* Section : NiagaraTrack->GetAllSections())
-					{
-						if (UMovieSceneNiagaraCacheSection* NiagaraCacheSection = Cast<UMovieSceneNiagaraCacheSection>(Section))
-						{
-							if (NiagaraCacheSection->bCacheOutOfDate)
-							{
-								return EVisibility::Visible;
-							}
-						}
-					}
+					return EVisibility::Visible;
 				}
 				return EVisibility::Collapsed;
 			})
diff --git a/Engine/Plugins/FX/NiagaraSimCaching/Source/NiagaraSimCachingEditor/Public/Sequencer/NiagaraCacheTrackEditor.h b/Engine/Plugins/FX/NiagaraSimCaching/Source/NiagaraSimCachingEditor/Public/Sequencer/NiagaraCacheTrackEditor.h
--- a/Engine/Plugins/FX/NiagaraSimCaching/Source/NiagaraSimCachingEditor/Public/Sequencer/NiagaraCacheTrackEditor.h
+++ b/Engine/Plugins/FX/NiagaraSimCaching/Source/NiagaraSimCachingEditor/Public/Sequencer/NiagaraCacheTrackEditor.h
@@ -53,6 +53,10 @@ private:
 	FKeyPropertyResult AddKeyInternal(FFrameNumber KeyTime, class UNiagaraComponent* NiagaraComponent, UMovieSceneTrack* Track);
 
 	FReply RecordCacheTrack(IMovieSceneCachedTrack* Track);
+
+	static bool IsCacheTrackOutOfDate(UMovieSceneTrack* Track);
+	bool HasConflictingLifecycleTrack(UMovieSceneTrack* Track) const;
+	FText GetCacheTrackWarnings(UMovieSceneTrack* Track) const;
 };
 
 /** Niagara cache sequencer section */
