Niagara: Fixed layout fallout from recent Sequencer column improvements

diff --git a/Engine/Plugins/FX/Niagara/Source/NiagaraEditor/NiagaraEditor.Build.cs b/Engine/Plugins/FX/Niagara/Source/NiagaraEditor/NiagaraEditor.Build.cs
--- a/Engine/Plugins/FX/Niagara/Source/NiagaraEditor/NiagaraEditor.Build.cs
+++ b/Engine/Plugins/FX/Niagara/Source/NiagaraEditor/NiagaraEditor.Build.cs
@@ -24,6 +24,7 @@ public class NiagaraEditor : ModuleRules
                 "EditorStyle",
 				"VectorVM",
                 "MovieScene",
+				"SequencerCore",
 				"Sequencer",
 				"TimeManagement",
 				"PropertyEditor",
diff --git a/Engine/Plugins/FX/Niagara/Source/NiagaraEditor/Private/Sequencer/LevelSequence/NiagaraSystemTrackEditor.h b/Engine/Plugins/FX/Niagara/Source/NiagaraEditor/Private/Sequencer/LevelSequence/NiagaraSystemTrackEditor.h
--- a/Engine/Plugins/FX/Niagara/Source/NiagaraEditor/Private/Sequencer/LevelSequence/NiagaraSystemTrackEditor.h
+++ b/Engine/Plugins/FX/Niagara/Source/NiagaraEditor/Private/Sequencer/LevelSequence/NiagaraSystemTrackEditor.h
@@ -41,7 +41,7 @@ public:
 	virtual TSharedRef<ISequencerSection> MakeSectionInterface(UMovieSceneSection& SectionObject, UMovieSceneTrack& Track, FGuid ObjectBinding) override;
 	virtual bool SupportsType(TSubclassOf<UMovieSceneTrack> Type) const override;
 	virtual void BuildObjectBindingTrackMenu(FMenuBuilder& MenuBuilder, const TArray<FGuid>& ObjectBindings, const UClass* ObjectClass) override;
-	virtual TSharedPtr<SWidget> BuildOutlinerEditWidget(const FGuid& ObjectBinding, UMovieSceneTrack* Track, const FBuildEditWidgetParams& Params) override;
+	virtual TSharedPtr<SWidget> BuildOutlinerColumnWidget(const FBuildColumnWidgetParams& Params, const FName& ColumnName) override;
 
 private:
 	void AddNiagaraSystemTrack(TArray<FGuid> ObjectBindings);
diff --git a/Engine/Plugins/FX/NiagaraSimCaching/Source/NiagaraSimCachingEditor/NiagaraSimCachingEditor.Build.cs b/Engine/Plugins/FX/NiagaraSimCaching/Source/NiagaraSimCachingEditor/NiagaraSimCachingEditor.Build.cs
--- a/Engine/Plugins/FX/NiagaraSimCaching/Source/NiagaraSimCachingEditor/NiagaraSimCachingEditor.Build.cs
+++ b/Engine/Plugins/FX/NiagaraSimCaching/Source/NiagaraSimCachingEditor/NiagaraSimCachingEditor.Build.cs
@@ -24,6 +24,7 @@ namespace UnrealBuildTool.Rules
                     "MovieScene",
                     "MovieSceneTools",
                     "MovieSceneTracks",
+                    "SequencerCore",
                     "Sequencer",
                     "TimeManagement",
 					"EditorFramework",
diff --git a/Engine/Plugins/FX/NiagaraSimCaching/Source/NiagaraSimCachingEditor/Private/Sequencer/NiagaraCacheTrackEditor.cpp b/Engine/Plugins/FX/NiagaraSimCaching/Source/NiagaraSimCachingEditor/Private/Sequencer/NiagaraCacheTrackEditor.cpp
--- a/Engine/Plugins/FX/NiagaraSimCaching/Source/NiagaraSimCachingEditor/Private/Sequencer/NiagaraCacheTrackEditor.cpp
+++ b/Engine/Plugins/FX/NiagaraSimCaching/Source/NiagaraSimCachingEditor/Private/Sequencer/NiagaraCacheTrackEditor.cpp
@@ -26,6 +26,9 @@
 #include "Widgets/Images/SImage.h"
 #include "Widgets/Input/SButton.h"
 #include "MVVM/ViewModels/ViewDensity.h"
+#include "MVVM/ViewModels/OutlinerColumns/OutlinerColumnTypes.h"
+#include "MVVM/Extensions/ITrackExtension.h"
+#include "MVVM/Extensions/IObjectBindingExtension.h"
 
 namespace NiagaraCacheEditorConstants
 {
@@ -497,12 +500,21 @@ FText FNiagaraCacheTrackEditor::GetCacheTrackWarnings(UMovieSceneTrack* Track) c
 	return FText::Join(FText::FromString("\n\n"), Warnings);
 }
 
-TSharedPtr<SWidget> FNiagaraCacheTrackEditor::BuildOutlinerEditWidget(const FGuid& ObjectBinding, UMovieSceneTrack* Track, const FBuildEditWidgetParams& Params)
+TSharedPtr<SWidget> FNiagaraCacheTrackEditor::BuildOutlinerColumnWidget(const FBuildColumnWidgetParams& Params, const FName& ColumnName)
 {
-	TSharedPtr<ISequencer> SequencerPtr = GetSequencer();
-	FWeakObjectPtr BoundObject = SequencerPtr.IsValid() ? SequencerPtr->FindSpawnedObjectOrTemplate(ObjectBinding) : nullptr;
-	
-	return SNew(SHorizontalBox)
+	using namespace UE::Sequencer;
+
+	if (ColumnName == FCommonOutlinerNames::Edit)
+	{
+		TViewModelPtr<IObjectBindingExtension> ObjectBindingModel = Params.ViewModel->FindAncestorOfType<IObjectBindingExtension>();
+		FGuid ObjectBinding = ObjectBindingModel ? ObjectBindingModel->GetObjectGuid() : FGuid();
+
+		TSharedPtr<ISequencer> SequencerPtr = GetSequencer();
+		FWeakObjectPtr BoundObject = SequencerPtr.IsValid() ? SequencerPtr->FindSpawnedObjectOrTemplate(ObjectBinding) : nullptr;
+		
+		UMovieSceneTrack* Track = Params.TrackModel->GetTrack();
+
+		return SNew(SHorizontalBox)
 		+ SHorizontalBox::Slot()
 		.Padding(2, 0, 2, 0)
 		.VAlign(VAlign_Center)
@@ -529,6 +541,7 @@ TSharedPtr<SWidget> FNiagaraCacheTrackEditor::BuildOutlinerEditWidget(const FGui
 		+ SHorizontalBox::Slot()
 		.Padding(2, 0, 2, 0)
 		.VAlign(VAlign_Center)
+		.AutoWidth()
 		[
 			SNew(SImage)
 			.ToolTipText_Raw(this, &FNiagaraCacheTrackEditor::GetCacheTrackWarnings, Track)
@@ -545,6 +558,7 @@ TSharedPtr<SWidget> FNiagaraCacheTrackEditor::BuildOutlinerEditWidget(const FGui
 		+ SHorizontalBox::Slot()
 		.Padding(2, 0, 2, 0)
 		.VAlign(VAlign_Center)
+		.AutoWidth()
 		[
 			SNew(SImage)
 			.ToolTipText(LOCTEXT("AddNiagaraCache_InvalidCache", "The cache used by this track does not fit the Niagara component and will be ignored. Re-Record the sim cache or select a different one."))
@@ -574,6 +588,7 @@ TSharedPtr<SWidget> FNiagaraCacheTrackEditor::BuildOutlinerEditWidget(const FGui
 		+ SHorizontalBox::Slot()
 		.Padding(2, 0, 2, 0)
 		.VAlign(VAlign_Center)
+		.AutoWidth()
 		[
 			SNew(SImage)
 			.ToolTipText(LOCTEXT("NiagaraCache_Active", "This track is currently providing cached data to the Niagara component, no simulation is running."))
@@ -601,6 +616,9 @@ TSharedPtr<SWidget> FNiagaraCacheTrackEditor::BuildOutlinerEditWidget(const FGui
 				return EVisibility::Collapsed;
 			})
 		];
+	}
+
+	return FMovieSceneTrackEditor::BuildOutlinerColumnWidget(Params, ColumnName);
 }
 
 const FSlateBrush* FNiagaraCacheTrackEditor::GetIconBrush() const
diff --git a/Engine/Plugins/FX/NiagaraSimCaching/Source/NiagaraSimCachingEditor/Public/Sequencer/NiagaraCacheTrackEditor.h b/Engine/Plugins/FX/NiagaraSimCaching/Source/NiagaraSimCachingEditor/Public/Sequencer/NiagaraCacheTrackEditor.h
--- a/Engine/Plugins/FX/NiagaraSimCaching/Source/NiagaraSimCachingEditor/Public/Sequencer/NiagaraCacheTrackEditor.h
+++ b/Engine/Plugins/FX/NiagaraSimCaching/Source/NiagaraSimCachingEditor/Public/Sequencer/NiagaraCacheTrackEditor.h
@@ -40,7 +40,7 @@ public:
 	virtual TSharedRef<ISequencerSection> MakeSectionInterface(UMovieSceneSection& SectionObject, UMovieSceneTrack& Track, FGuid ObjectBinding) override;
 	virtual bool SupportsSequence(UMovieSceneSequence* InSequence) const override;
 	virtual bool SupportsType(TSubclassOf<UMovieSceneTrack> Type) const override;
-	virtual TSharedPtr<SWidget> BuildOutlinerEditWidget(const FGuid& ObjectBinding, UMovieSceneTrack* Track, const FBuildEditWidgetParams& Params) override;
+	virtual TSharedPtr<SWidget> BuildOutlinerColumnWidget(const FBuildColumnWidgetParams& Params, const FName& ColumnName) override;
 	virtual const FSlateBrush* GetIconBrush() const override;
 	// ~End ISequencerTrackEditor interface
 	
