Fixed missing struct default value

diff --git a/Engine/Plugins/FX/NiagaraSimCaching/Source/NiagaraSimCaching/Public/Niagara/Sequencer/MovieSceneNiagaraCacheSection.h b/Engine/Plugins/FX/NiagaraSimCaching/Source/NiagaraSimCaching/Public/Niagara/Sequencer/MovieSceneNiagaraCacheSection.h
--- a/Engine/Plugins/FX/NiagaraSimCaching/Source/NiagaraSimCaching/Public/Niagara/Sequencer/MovieSceneNiagaraCacheSection.h
+++ b/Engine/Plugins/FX/NiagaraSimCaching/Source/NiagaraSimCaching/Public/Niagara/Sequencer/MovieSceneNiagaraCacheSection.h
@@ -62,7 +62,7 @@ struct NIAGARASIMCACHING_API FMovieSceneNiagaraCacheParams : public FMovieSceneB
 
 	/** If set, then the engine scalability setting will be overriden with this value when recording a new cache for this track */
 	UPROPERTY(EditAnywhere, Category = "NiagaraCache", meta=(EditCondition="bOverrideQualityLevel"))
-	EPerQualityLevels RecordQualityLevel;
+	EPerQualityLevels RecordQualityLevel = EPerQualityLevels::Cinematic;
 
 	/** What should the effect do when the track has no cache data to display */
 	UPROPERTY(EditAnywhere, Category="SimCache")
