Remove "experimental" tag from Niagara sim caching plugin and enable by default

diff --git a/Engine/Plugins/FX/NiagaraSimCaching/NiagaraSimCaching.uplugin b/Engine/Plugins/FX/NiagaraSimCaching/NiagaraSimCaching.uplugin
--- a/Engine/Plugins/FX/NiagaraSimCaching/NiagaraSimCaching.uplugin
+++ b/Engine/Plugins/FX/NiagaraSimCaching/NiagaraSimCaching.uplugin
@@ -7,13 +7,13 @@
 	"Category": "FX",
 	"CreatedBy": "Epic Games, Inc.",
 	"CreatedByURL": "https://epicgames.com",
-	"DocsURL": "",
+	"DocsURL": "https://dev.epicgames.com/community/learning/tutorials/Rk9v/unreal-engine-niagara-simulation-caching-in-sequencer",
 	"MarketplaceURL": "",
 	"SupportURL": "",
-	"EnabledByDefault": false,
+	"EnabledByDefault": true,
 	"CanContainContent": true,
-	"IsBetaVersion": true,
-	"IsExperimentalVersion": true,
+	"IsBetaVersion": false,
+	"IsExperimentalVersion": false,
 	"Installed": false,
 	"Modules": [
 		{
