- Fix for take recorder caching being a frame out
- Take recorder ticks in the level tickables and sequencer ticks in the editor tickables
- This meant that we were always recording the previous frame and not the current one
- We now record via OnPostEditorTick which guarantees ordering

diff --git a/Engine/Plugins/FX/NiagaraSimCaching/Source/NiagaraSimCachingEditor/Private/Sequencer/MovieSceneNiagaraTrackRecorder.cpp b/Engine/Plugins/FX/NiagaraSimCaching/Source/NiagaraSimCachingEditor/Private/Sequencer/MovieSceneNiagaraTrackRecorder.cpp
--- a/Engine/Plugins/FX/NiagaraSimCaching/Source/NiagaraSimCachingEditor/Private/Sequencer/MovieSceneNiagaraTrackRecorder.cpp
+++ b/Engine/Plugins/FX/NiagaraSimCaching/Source/NiagaraSimCachingEditor/Private/Sequencer/MovieSceneNiagaraTrackRecorder.cpp
@@ -155,10 +155,6 @@ void UMovieSceneNiagaraTrackRecorder::SetSectionStartTimecodeImpl(const FTimecod
 		FTakeRecorderParameters Parameters;
 		Parameters.User = GetDefault<UTakeRecorderUserSettings>()->Settings;
 		Parameters.Project = GetDefault<UTakeRecorderProjectSettings>()->Settings;
-		RecordStartTime = FApp::GetCurrentTime();
-		RecordStartFrame = Parameters.Project.bStartAtCurrentTimecode ?
-			FFrameRate::TransformTime(FFrameTime(InSectionStartTimecode.ToFrameNumber(DisplayRate)),
-				DisplayRate, TickResolution).FloorToFrame() : MovieScene->GetPlaybackRange().GetLowerBoundValue();
 
 		if (RecordRange.IsSet())
 		{
@@ -182,6 +178,11 @@ void UMovieSceneNiagaraTrackRecorder::SetSectionStartTimecodeImpl(const FTimecod
 				NiagaraCacheSection->Params.SimCache = NewObject<UNiagaraSimCache>(NiagaraCacheSection.Get(), NAME_None, RF_Transactional);
 			}
 			NiagaraCacheSection->Params.SimCache->BeginWrite(NiagaraCacheSection->Params.CacheParameters, SystemToRecord.Get());
+
+			bRecordedFirstFrame = false;
+			bRecordingEnabled = false;
+			bRequestFinalize = false;
+			PostEditorTickHandle = GEngine->OnPostEditorTick().AddUObject(this, &UMovieSceneNiagaraTrackRecorder::OnRecordFrame);
 			NiagaraCacheSection->bCacheOutOfDate = false;
 		}
 	}
@@ -194,31 +195,20 @@ UMovieSceneSection* UMovieSceneNiagaraTrackRecorder::GetMovieSceneSection() cons
 
 void UMovieSceneNiagaraTrackRecorder::FinalizeTrackImpl()
 {
-	if (NiagaraCacheTrack.IsValid())
-	{
-		NiagaraCacheTrack->bIsRecording = false;
-	}
-	if (NiagaraCacheSection.IsValid())
-	{
-		// finalize the sim cache
-		NiagaraCacheSection->Params.SimCache->EndWrite(true);
-		
-		// Activate the section
-		NiagaraCacheSection->SetIsActive(true);
-	}
+	bRequestFinalize = true;
 }
 
 void UMovieSceneNiagaraTrackRecorder::RecordSampleImpl(const FQualifiedFrameTime& CurrentFrameTime)
 {
 	const FFrameRate TickResolution = MovieScene->GetTickResolution();
-	const FFrameNumber CurrentFrame = CurrentFrameTime.ConvertTo(TickResolution).FloorToFrame();
+	RecordingFrameNumber = CurrentFrameTime.ConvertTo(TickResolution).FloorToFrame();
 
-	if (RecordRange.IsSet() && (RecordRange->GetLowerBoundValue() > CurrentFrame || RecordRange->GetUpperBoundValue() <= CurrentFrame))
-	{
-		return;
-	}
+	bRecordingEnabled = !RecordRange.IsSet() || (RecordingFrameNumber >= RecordRange->GetLowerBoundValue() && RecordingFrameNumber < RecordRange->GetUpperBoundValue());
+}
 
-	if (NiagaraCacheSection.IsValid() && SystemToRecord.IsValid())
+void UMovieSceneNiagaraTrackRecorder::OnRecordFrame(float DeltaSeconds)
+{
+	if (bRecordingEnabled && NiagaraCacheSection.IsValid() && SystemToRecord.IsValid())
 	{
 		FNiagaraSimCacheFeedbackContext FeedbackContext;
 		FeedbackContext.bAutoLogIssues = false;
@@ -227,23 +217,42 @@ void UMovieSceneNiagaraTrackRecorder::RecordSampleImpl(const FQualifiedFrameTime
 			if (!bRecordedFirstFrame)
 			{
 				// set to the actual first recorded frame, because systems with spawn rate can tick for a few frames without having particles
-				NiagaraCacheSection->SetStartFrame(CurrentFrame);
+				NiagaraCacheSection->SetStartFrame(RecordingFrameNumber);
 				bRecordedFirstFrame = true;
 			}
-			
+
 			// Expand the section to the new length
-			FFrameNumber EndFrame = CurrentFrameTime.Time.CeilToFrame();
-			NiagaraCacheSection->SetEndFrame(EndFrame);
+			NiagaraCacheSection->SetEndFrame(RecordingFrameNumber);
 		}
-		
+
 		for (const FString& Warning : FeedbackContext.Warnings)
 		{
-			UE_LOG(LogNiagaraSimCachingEditor, Warning, TEXT("Recording sim cache for frame %i: %s"), CurrentFrame.Value, *Warning);
+			UE_LOG(LogNiagaraSimCachingEditor, Warning, TEXT("Recording sim cache for frame %i: %s"), RecordingFrameNumber.Value, *Warning);
 		}
 
 		for (const FString& Error : FeedbackContext.Errors)
 		{
-			UE_LOG(LogNiagaraSimCachingEditor, Warning, TEXT("Unable to record sim cache for frame %i: %s"), CurrentFrame.Value, *Error);
+			UE_LOG(LogNiagaraSimCachingEditor, Warning, TEXT("Unable to record sim cache for frame %i: %s"), RecordingFrameNumber.Value, *Error);
+		}
+	}
+
+	if (bRequestFinalize)
+	{
+		GEngine->OnPostEditorTick().Remove(PostEditorTickHandle);
+		PostEditorTickHandle.Reset();
+
+		if (NiagaraCacheTrack.IsValid())
+		{
+			NiagaraCacheTrack->bIsRecording = false;
+		}
+
+		if (NiagaraCacheSection.IsValid())
+		{
+			// Finalize the sim cache
+			NiagaraCacheSection->Params.SimCache->EndWrite(true);
+
+			// Activate the section
+			NiagaraCacheSection->SetIsActive(true);
 		}
 	}
 }
diff --git a/Engine/Plugins/FX/NiagaraSimCaching/Source/NiagaraSimCachingEditor/Public/Sequencer/MovieSceneNiagaraTrackRecorder.h b/Engine/Plugins/FX/NiagaraSimCaching/Source/NiagaraSimCachingEditor/Public/Sequencer/MovieSceneNiagaraTrackRecorder.h
--- a/Engine/Plugins/FX/NiagaraSimCaching/Source/NiagaraSimCachingEditor/Public/Sequencer/MovieSceneNiagaraTrackRecorder.h
+++ b/Engine/Plugins/FX/NiagaraSimCaching/Source/NiagaraSimCachingEditor/Public/Sequencer/MovieSceneNiagaraTrackRecorder.h
@@ -50,6 +50,10 @@ public:
 
 	/** Returns the Niagara cache track on which the cache manager will be recorded */
 	TWeakObjectPtr<UMovieSceneNiagaraCacheTrack> GetNiagaraCacheTrack() const {return NiagaraCacheTrack;}
+
+protected:
+	void SetRecordingEnabled(bool bEnabled);
+	void OnRecordFrame(float DeltaSeconds);
 	
 private:
 	
@@ -62,14 +66,13 @@ private:
 	/** Object to record from */
 	TLazyObjectPtr<UNiagaraComponent> SystemToRecord;
 
-	/** The time at the start of this recording section */
-	double RecordStartTime;
-
-	/** The frame at the start of this recording section */
-	FFrameNumber RecordStartFrame;
-
 	TOptional<FFrameNumberRange> RecordRange;
-	bool bRecordedFirstFrame = false;
+
+	bool			bRecordedFirstFrame = false;
+	bool			bRecordingEnabled = false;
+	bool			bRequestFinalize = false;
+	FFrameNumber	RecordingFrameNumber;
+	FDelegateHandle	PostEditorTickHandle;
 
 	friend FMovieSceneNiagaraTrackRecorderFactory;
 };
