Added a setting to record niagara systems at a rate lower than the sequencer play rate.

#tests Created a level sequence with a high playback rate (240fps), setting the cache to record at a low rate (24fps). Checked that the number of recorded frames was the same as if the sequence was playing at the low rate (24fps).

diff --git a/Engine/Plugins/FX/NiagaraSimCaching/Source/NiagaraSimCaching/Public/Niagara/Sequencer/MovieSceneNiagaraCacheSection.h b/Engine/Plugins/FX/NiagaraSimCaching/Source/NiagaraSimCaching/Public/Niagara/Sequencer/MovieSceneNiagaraCacheSection.h
--- a/Engine/Plugins/FX/NiagaraSimCaching/Source/NiagaraSimCaching/Public/Niagara/Sequencer/MovieSceneNiagaraCacheSection.h
+++ b/Engine/Plugins/FX/NiagaraSimCaching/Source/NiagaraSimCaching/Public/Niagara/Sequencer/MovieSceneNiagaraCacheSection.h
@@ -71,6 +71,19 @@ struct NIAGARASIMCACHING_API FMovieSceneNiagaraCacheParams : public FMovieSceneB
 	/** What should the effect do when the cache section is stretched? */
 	UPROPERTY(EditAnywhere, Category="SimCache")
 	ENiagaraSimCacheSectionStretchMode SectionStretchMode = ENiagaraSimCacheSectionStretchMode::TimeDilate;
+
+#if WITH_EDITORONLY_DATA 
+	/** True if the cache should be recorded at a rate that is slower than the sequencer play rate. */
+	UPROPERTY(EditAnywhere, Category = "Cache", meta = (InlineEditConditionToggle))
+	bool bOverrideRecordRate = false;
+
+	/** The rate at which the cache should be recorded. Will be ignored if the sequence plays at a lower rate than this number.
+	 * This option is useful if you need to play the sequence at a very high fps for simulation stability (e.g. 240fps), but then later want to play back the cache at a normal rate (e.g. 24fps).
+	 * Please note that using this option might result in rendering artifacts for things like motion blur, velocity interpolation or inconsistent grid data in fluids, due to the dropped cache frames. 
+	 */
+	UPROPERTY(EditAnywhere, Category = "Cache", meta = (EditCondition = "bOverrideRecordRate"))
+	float CacheRecordRateFPS = 24.0f;
+#endif
 };
 
 /**
diff --git a/Engine/Plugins/FX/NiagaraSimCaching/Source/NiagaraSimCachingEditor/Private/Sequencer/MovieSceneNiagaraTrackRecorder.cpp b/Engine/Plugins/FX/NiagaraSimCaching/Source/NiagaraSimCachingEditor/Private/Sequencer/MovieSceneNiagaraTrackRecorder.cpp
--- a/Engine/Plugins/FX/NiagaraSimCaching/Source/NiagaraSimCachingEditor/Private/Sequencer/MovieSceneNiagaraTrackRecorder.cpp
+++ b/Engine/Plugins/FX/NiagaraSimCaching/Source/NiagaraSimCachingEditor/Private/Sequencer/MovieSceneNiagaraTrackRecorder.cpp
@@ -203,7 +203,21 @@ void UMovieSceneNiagaraTrackRecorder::RecordSampleImpl(const FQualifiedFrameTime
 	const FFrameRate TickResolution = MovieScene->GetTickResolution();
 	RecordingFrameNumber = CurrentFrameTime.ConvertTo(TickResolution).FloorToFrame();
 
-	bRecordingEnabled = !RecordRange.IsSet() || (RecordingFrameNumber >= RecordRange->GetLowerBoundValue() && RecordingFrameNumber < RecordRange->GetUpperBoundValue());
+	bool bRecordCurrentFrame = true;
+	if (NiagaraCacheSection.IsValid() && NiagaraCacheSection->Params.bOverrideRecordRate)
+	{
+		if (LastRecordedFrame.Value > 0 && NiagaraCacheSection->Params.CacheRecordRateFPS != 0)
+		{
+			double ElapsedTime = TickResolution.AsInterval() * (RecordingFrameNumber.Value - LastRecordedFrame.Value);
+			double DesiredTime = 1.0 / NiagaraCacheSection->Params.CacheRecordRateFPS;
+			if (ElapsedTime < DesiredTime && !FMath::IsNearlyEqual(ElapsedTime,DesiredTime))
+			{
+				bRecordCurrentFrame = false;
+			}
+		}
+	}
+
+	bRecordingEnabled = bRecordCurrentFrame && (!RecordRange.IsSet() || (RecordingFrameNumber >= RecordRange->GetLowerBoundValue() && RecordingFrameNumber < RecordRange->GetUpperBoundValue()));
 }
 
 void UMovieSceneNiagaraTrackRecorder::OnRecordFrame(float DeltaSeconds)
@@ -223,6 +237,8 @@ void UMovieSceneNiagaraTrackRecorder::OnRecordFrame(float DeltaSeconds)
 
 			// Expand the section to the new length
 			NiagaraCacheSection->SetEndFrame(RecordingFrameNumber);
+			
+			LastRecordedFrame = RecordingFrameNumber;
 		}
 
 		for (const FString& Warning : FeedbackContext.Warnings)
diff --git a/Engine/Plugins/FX/NiagaraSimCaching/Source/NiagaraSimCachingEditor/Public/Sequencer/MovieSceneNiagaraTrackRecorder.h b/Engine/Plugins/FX/NiagaraSimCaching/Source/NiagaraSimCachingEditor/Public/Sequencer/MovieSceneNiagaraTrackRecorder.h
--- a/Engine/Plugins/FX/NiagaraSimCaching/Source/NiagaraSimCachingEditor/Public/Sequencer/MovieSceneNiagaraTrackRecorder.h
+++ b/Engine/Plugins/FX/NiagaraSimCaching/Source/NiagaraSimCachingEditor/Public/Sequencer/MovieSceneNiagaraTrackRecorder.h
@@ -71,6 +71,7 @@ private:
 	bool			bRecordedFirstFrame = false;
 	bool			bRecordingEnabled = false;
 	bool			bRequestFinalize = false;
+	FFrameNumber	LastRecordedFrame;
 	FFrameNumber	RecordingFrameNumber;
 	FDelegateHandle	PostEditorTickHandle;
 
